Initial submit of the new JSON-based editor config system. This consists of two layers - the JSON config system enables access to JSON objects through the use of JSON paths of the form $.Object.Array[3].Member. The editor config system applies this system to UObjects and UStructs.

diff --git a/Engine/Source/Editor/EditorConfig/EditorConfig.Build.cs b/Engine/Source/Editor/EditorConfig/EditorConfig.Build.cs
new file mode 100644
--- /dev/null
+++ b/Engine/Source/Editor/EditorConfig/EditorConfig.Build.cs
@@ -0,0 +1,21 @@
+// Copyright Epic Games, Inc. All Rights Reserved.
+
+namespace UnrealBuildTool.Rules
+{
+	public class EditorConfig : ModuleRules
+	{
+        public EditorConfig(ReadOnlyTargetRules Target) : base(Target)
+		{
+			PrivateDependencyModuleNames.AddRange(
+                new string[]
+                {
+                    "Core",
+					"CoreUObject",
+					"EditorSubsystem",
+					"Engine",
+					"Json",
+					"UnrealEd",
+				} );
+        }
+	}
+}
diff --git a/Engine/Source/Editor/EditorConfig/Private/EditorConfig.cpp b/Engine/Source/Editor/EditorConfig/Private/EditorConfig.cpp
new file mode 100644
--- /dev/null
+++ b/Engine/Source/Editor/EditorConfig/Private/EditorConfig.cpp
@@ -0,0 +1,631 @@
+// Copyright Epic Games, Inc. All Rights Reserved.
+
+#include "EditorConfig.h"
+#include "UObject/UnrealType.h"
+
+FEditorConfig::FEditorConfig()
+{
+	JsonConfig = MakeShared<UE::FJsonConfig>();
+}
+
+void FEditorConfig::SetParent(TSharedPtr<FEditorConfig> InConfig)
+{
+	ParentConfig = InConfig;
+
+	if (ParentConfig.IsValid())
+	{
+		JsonConfig->SetParent(ParentConfig->JsonConfig);
+	}
+	else
+	{
+		JsonConfig->SetParent(TSharedPtr<UE::FJsonConfig>());
+	}
+}
+
+bool FEditorConfig::LoadFromFile(FStringView FilePath)
+{
+	JsonConfig = MakeShared<UE::FJsonConfig>();
+	if (!JsonConfig->LoadFromFile(FilePath))
+	{
+		return false;
+	}
+
+	if (ParentConfig.IsValid())
+	{
+		JsonConfig->SetParent(ParentConfig->JsonConfig);
+	}
+
+	return true;
+}
+	
+bool FEditorConfig::LoadFromString(FStringView Content)
+{
+	JsonConfig = MakeShared<UE::FJsonConfig>();
+	if (!JsonConfig->LoadFromString(Content))
+	{
+		return false;
+	}
+
+	if (ParentConfig.IsValid())
+	{
+		JsonConfig->SetParent(ParentConfig->JsonConfig);
+	}
+
+	return true;
+}
+
+bool FEditorConfig::SaveToString(FString& OutResult) const
+{
+	if (!IsValid())
+	{
+		return false;
+	}
+
+	return JsonConfig->SaveToString(OutResult);
+}
+
+bool FEditorConfig::SaveToFile(FStringView FilePath) const
+{
+	if (!IsValid())
+	{
+		return false;
+	}
+
+	return JsonConfig->SaveToFile(FilePath);
+}
+
+bool FEditorConfig::HasOverride(FStringView Key) const
+{
+	return JsonConfig->HasOverride(UE::FJsonPath(Key));
+}
+
+void FEditorConfig::ReadStruct(TSharedPtr<FJsonObject> JsonObject, const UStruct* Struct, void* Instance, UObject* Owner)
+{
+	FString TypeName;
+	JsonObject->TryGetStringField(TEXT("$type"), TypeName);
+
+	if (!ensureAlways(Struct->GetName() != TypeName))
+	{
+		return;
+	}
+
+	for (TFieldIterator<FProperty> It(Struct); It; ++It)
+	{
+		FProperty* Property = *It;
+		if (Property != nullptr)
+		{
+			void* DataPtr = Property->ContainerPtrToValuePtr<void>(Instance);
+			
+			TSharedPtr<FJsonValue> Value = JsonObject->TryGetField(Property->GetName());
+			if (Value.IsValid())
+			{
+				ReadValue(Value, Property, DataPtr, Owner);
+			}
+		}
+	}
+}
+
+void FEditorConfig::ReadUObject(TSharedPtr<FJsonObject> JsonObject, const UClass* Class, UObject* Instance)
+{
+	FString TypeName;
+	JsonObject->TryGetStringField(TEXT("$type"), TypeName);
+
+	if (!ensureAlways(Class->GetName() != TypeName))
+	{
+		return;
+	}
+
+	for (TFieldIterator<FProperty> It(Class); It; ++It)
+	{
+		FProperty* Property = *It;
+		if (Property != nullptr)
+		{
+			void* DataPtr = Property->ContainerPtrToValuePtr<void>(Instance);
+
+			TSharedPtr<FJsonValue> Value = JsonObject->TryGetField(Property->GetName());
+			if (Value.IsValid())
+			{
+				ReadValue(Value, Property, DataPtr, Instance);
+			}
+		}
+	}
+}
+
+void FEditorConfig::ReadValue(TSharedPtr<FJsonValue> JsonValue, const FProperty* Property, void* DataPtr, UObject* Owner)
+{
+	if (const FStrProperty* StrProperty = CastField<FStrProperty>(Property))
+	{
+		FString* Value = (FString*) DataPtr;
+		JsonValue->TryGetString(*Value);
+	}
+	else if (const FNameProperty* NameProperty = CastField<FNameProperty>(Property))
+	{
+		FString TempValue;
+		JsonValue->TryGetString(TempValue);
+
+		*(FName*) DataPtr = *TempValue;
+	}
+	else if (const FTextProperty* TextProperty = CastField<FTextProperty>(Property))
+	{
+		FString TempValue;
+		JsonValue->TryGetString(TempValue);
+
+		*(FText*) DataPtr = FText::FromString(TempValue);
+	}
+	else if (const FBoolProperty* BoolProperty = CastField<FBoolProperty>(Property))
+	{
+		bool Value = BoolProperty->GetDefaultPropertyValue();
+		if (JsonValue->TryGetBool(Value))
+		{
+			BoolProperty->SetPropertyValue(DataPtr, Value);
+		}
+	}
+	else if (const FFloatProperty* FloatProperty = CastField<FFloatProperty>(Property))
+	{
+		float* Value = (float*) DataPtr;
+		JsonValue->TryGetNumber(*Value);
+	}
+	else if (const FDoubleProperty* DoubleProperty = CastField<FDoubleProperty>(Property))
+	{
+		double* Value = (double*) DataPtr;
+		JsonValue->TryGetNumber(*Value);
+	} 
+	else if (const FInt8Property* Int8Property = CastField<FInt8Property>(Property))
+	{
+		int8* Value = (int8*) DataPtr;
+		JsonValue->TryGetNumber(*Value);
+	}
+	else if (const FInt16Property* Int16Property = CastField<FInt16Property>(Property))
+	{
+		int16* Value = (int16*) DataPtr;
+		JsonValue->TryGetNumber(*Value);
+	}
+	else if (const FIntProperty* Int32Property = CastField<FIntProperty>(Property))
+	{
+		int32* Value = (int32*) DataPtr;
+		JsonValue->TryGetNumber(*Value);
+	}
+	else if (const FInt64Property* Int64Property = CastField<FInt64Property>(Property))
+	{
+		int64* Value = (int64*) DataPtr;
+		JsonValue->TryGetNumber(*Value);
+	}
+	else if (const FByteProperty* ByteProperty = CastField<FByteProperty>(Property))
+	{
+		uint8* Value = (uint8*) DataPtr;
+		JsonValue->TryGetNumber(*Value);
+	}
+	else if (const FUInt16Property* Uint16Property = CastField<FUInt16Property>(Property))
+	{
+		uint16* Value = (uint16*) DataPtr;
+		JsonValue->TryGetNumber(*Value);
+	}
+	else if (const FUInt32Property* Uint32Property = CastField<FUInt32Property>(Property))
+	{
+		uint32* Value = (uint32*) DataPtr;
+		JsonValue->TryGetNumber(*Value);
+	}
+	else if (const FUInt64Property* Uint64Property = CastField<FUInt64Property>(Property))
+	{
+		uint64* Value = (uint64*) DataPtr;
+		JsonValue->TryGetNumber(*Value);
+	}
+	else if (const FEnumProperty* EnumProperty = CastField<FEnumProperty>(Property))
+	{
+		int64* Value = (int64*) DataPtr;
+
+		UEnum* Enum = EnumProperty->GetEnum();
+		if (Enum != nullptr)
+		{
+			FString ValueString;
+			if (JsonValue->TryGetString(ValueString))
+			{
+				int64 Index = Enum->GetIndexByNameString(ValueString);
+				if (Index != INDEX_NONE)
+				{
+					*Value = Enum->GetValueByIndex(Index);
+				}
+			}
+		}
+	}
+	else if (const FObjectPropertyBase* ObjectProperty = CastField<FObjectPropertyBase>(Property))
+	{
+		FString PathString;
+		if (JsonValue->TryGetString(PathString))
+		{
+			Property->ImportText(*PathString, DataPtr, 0, Owner);
+		}
+	}
+	else if (const FStructProperty* StructProperty = CastField<FStructProperty>(Property))
+	{
+		const TSharedPtr<FJsonObject>* ObjectJsonValue;
+		if (JsonValue->TryGetObject(ObjectJsonValue))
+		{
+			ReadStruct(*ObjectJsonValue, StructProperty->Struct, DataPtr, Owner);
+		}
+	}
+	else if (const FArrayProperty* ArrayProperty = CastField<FArrayProperty>(Property))
+	{
+		const TArray<TSharedPtr<FJsonValue>>* ArrayJsonValue;
+		if (JsonValue->TryGetArray(ArrayJsonValue))
+		{
+			FProperty* InnerProperty = ArrayProperty->Inner;
+			FScriptArrayHelper ArrayHelper(ArrayProperty, DataPtr);
+
+			ArrayHelper.EmptyAndAddValues(ArrayJsonValue->Num());
+
+			for (int32 Idx = 0; Idx < ArrayHelper.Num(); ++Idx)
+			{
+				TSharedPtr<FJsonValue> Value = (*ArrayJsonValue)[Idx];
+				ReadValue(Value, InnerProperty, ArrayHelper.GetRawPtr(Idx), Owner);
+			}
+		}
+	}
+	else if (const FSetProperty* SetProperty = CastField<FSetProperty>(Property))
+	{
+		const TArray<TSharedPtr<FJsonValue>>* SetJsonValue;
+		if (JsonValue->TryGetArray(SetJsonValue))
+		{
+			const FProperty* InnerProperty = SetProperty->ElementProp;
+			FScriptSetHelper SetHelper(SetProperty, DataPtr);
+			SetHelper.EmptyElements(SetJsonValue->Num());
+
+			// temporary buffer to read elements into
+			TArray<uint8> TempBuffer;
+			TempBuffer.AddUninitialized(InnerProperty->ElementSize);
+
+			for (int32 Idx = 0; Idx < SetJsonValue->Num(); ++Idx)
+			{
+				InnerProperty->InitializeValue(TempBuffer.GetData());
+
+				TSharedPtr<FJsonValue> Value = (*SetJsonValue)[Idx];
+				ReadValue(Value, InnerProperty, TempBuffer.GetData(), Owner); 
+
+				SetHelper.AddElement(TempBuffer.GetData());
+
+				InnerProperty->DestroyValue(TempBuffer.GetData());
+			}
+		}
+	}
+	else if (const FMapProperty* MapProperty = CastField<FMapProperty>(Property))
+	{
+		const FProperty* KeyProperty = MapProperty->KeyProp;
+		const FProperty* ValueProperty = MapProperty->ValueProp;
+
+		FScriptMapHelper MapHelper(MapProperty, DataPtr);
+
+		// maps can either be stored as a simple JSON object, or as an array of { $key, $value } pairs
+		// first check for object storage - this will cover eg. numbers and strings as keys
+		const TSharedPtr<FJsonObject>* JsonObjectValue;
+		if (JsonValue->TryGetObject(JsonObjectValue))
+		{
+			MapHelper.EmptyValues((*JsonObjectValue)->Values.Num());
+
+			// temporary buffers to read elements into
+			TArray<uint8> TempKey;
+			TempKey.AddZeroed(KeyProperty->ElementSize);
+
+			TArray<uint8> TempValue;
+			TempValue.AddZeroed(ValueProperty->ElementSize);
+
+			for (const TPair<FString, TSharedPtr<FJsonValue>>& JsonPair : (*JsonObjectValue)->Values)
+			{
+				KeyProperty->InitializeValue(TempKey.GetData());
+				KeyProperty->ImportText(*JsonPair.Key, TempKey.GetData(), 0, Owner);
+
+				ValueProperty->InitializeValue(TempValue.GetData());
+				ReadValue(JsonPair.Value, ValueProperty, TempValue.GetData(), Owner); 
+
+				MapHelper.AddPair(TempKey.GetData(), TempValue.GetData());
+
+				KeyProperty->DestroyValue(TempKey.GetData());
+				ValueProperty->DestroyValue(TempValue.GetData());
+			}
+		
+			return;
+		}
+
+		// then check for array storage, this will cover complex keys eg. custom structs
+		const TArray<TSharedPtr<FJsonValue>>* JsonArrayPtr = nullptr;
+		if (JsonValue->TryGetArray(JsonArrayPtr))
+		{
+			MapHelper.EmptyValues(JsonArrayPtr->Num());
+			
+			// temporary buffers to read elements into
+			TArray<uint8> TempKey;
+			TempKey.AddUninitialized(KeyProperty->ElementSize);
+
+			TArray<uint8> TempValue;
+			TempValue.AddUninitialized(ValueProperty->ElementSize);
+
+			for (const TSharedPtr<FJsonValue>& JsonElement : *JsonArrayPtr)
+			{
+				TSharedPtr<FJsonObject>* JsonObject = nullptr;
+				if (JsonElement->TryGetObject(JsonObject))
+				{
+					TSharedPtr<FJsonValue> JsonKeyField = (*JsonObject)->TryGetField(TEXT("$key"));
+					TSharedPtr<FJsonValue> JsonValueField = (*JsonObject)->TryGetField(TEXT("$value"));
+
+					if (JsonKeyField.IsValid() && JsonValueField.IsValid())
+					{
+						KeyProperty->InitializeValue(TempKey.GetData());
+						ReadValue(JsonKeyField, KeyProperty, TempKey.GetData(), Owner); 
+
+						ValueProperty->InitializeValue(TempValue.GetData());
+						ReadValue(JsonValueField, ValueProperty, TempValue.GetData(), Owner); 
+
+						MapHelper.AddPair(TempKey.GetData(), TempValue.GetData());
+
+						KeyProperty->DestroyValue(TempKey.GetData());
+						ValueProperty->DestroyValue(TempValue.GetData());
+					}
+				}
+			}
+
+			return;
+		}
+	}
+}
+
+TSharedPtr<FJsonValue> FEditorConfig::WriteValue(const FProperty* Property, const void* DataPtr)
+{
+	TSharedPtr<FJsonValue> ResultValue;
+
+	if (const FStrProperty* StrProperty = CastField<FStrProperty>(Property))
+	{
+		FString* Value = (FString*) DataPtr;
+		ResultValue = MakeShared<FJsonValueString>(*Value);
+	}
+	else if (const FNameProperty* NameProperty = CastField<FNameProperty>(Property))
+	{
+		FName* Value = (FName*) DataPtr;
+		ResultValue = MakeShared<FJsonValueString>(Value->ToString());
+	}
+	else if (const FTextProperty* TextProperty = CastField<FTextProperty>(Property))
+	{
+		FText* Value = (FText*) DataPtr;
+		ResultValue = MakeShared<FJsonValueString>(Value->ToString());
+	}
+	else if (const FBoolProperty* BoolProperty = CastField<FBoolProperty>(Property))
+	{
+		bool Value = BoolProperty->GetPropertyValue(DataPtr);
+		ResultValue = MakeShared<FJsonValueBoolean>(Value);
+	}
+	else if (const FFloatProperty* FloatProperty = CastField<FFloatProperty>(Property))
+	{
+		float* Value = (float*) DataPtr;
+		ResultValue = MakeShared<FJsonValueNumber>(*Value);
+	}
+	else if (const FDoubleProperty* DoubleProperty = CastField<FDoubleProperty>(Property))
+	{
+		double* Value = (double*) DataPtr;
+		ResultValue = MakeShared<FJsonValueNumber>(*Value);
+	} 
+	else if (const FInt8Property* Int8Property = CastField<FInt8Property>(Property))
+	{
+		int8* Value = (int8*) DataPtr;
+		ResultValue = MakeShared<FJsonValueNumber>(*Value);
+	}
+	else if (const FInt16Property* Int16Property = CastField<FInt16Property>(Property))
+	{
+		int16* Value = (int16*) DataPtr;
+		ResultValue = MakeShared<FJsonValueNumber>(*Value);
+	}
+	else if (const FIntProperty* Int32Property = CastField<FIntProperty>(Property))
+	{
+		int32* Value = (int32*) DataPtr;
+		ResultValue = MakeShared<FJsonValueNumber>(*Value);
+	}
+	else if (const FInt64Property* Int64Property = CastField<FInt64Property>(Property))
+	{
+		int64* Value = (int64*) DataPtr;
+		ResultValue = MakeShared<FJsonValueNumber>(*Value);
+	}
+	else if (const FByteProperty* ByteProperty = CastField<FByteProperty>(Property))
+	{
+		uint8* Value = (uint8*) DataPtr;
+		ResultValue = MakeShared<FJsonValueNumber>(*Value);
+	}
+	else if (const FUInt16Property* Uint16Property = CastField<FUInt16Property>(Property))
+	{
+		uint16* Value = (uint16*) DataPtr;
+		ResultValue = MakeShared<FJsonValueNumber>(*Value);
+	}
+	else if (const FUInt32Property* Uint32Property = CastField<FUInt32Property>(Property))
+	{
+		uint32* Value = (uint32*) DataPtr;
+		ResultValue = MakeShared<FJsonValueNumber>(*Value);
+	}
+	else if (const FUInt64Property* Uint64Property = CastField<FUInt64Property>(Property))
+	{
+		uint64* Value = (uint64*) DataPtr;
+		ResultValue = MakeShared<FJsonValueNumber>(*Value);
+	}
+	else if (const FEnumProperty* EnumProperty = CastField<FEnumProperty>(Property))
+	{
+		int64* Value = (int64*) DataPtr;
+
+		UEnum* Enum = EnumProperty->GetEnum();
+		FName ValueName = Enum->GetNameByValue(*Value);
+		ResultValue = MakeShared<FJsonValueString>(ValueName.ToString());
+	}
+	else if (const FObjectPropertyBase* ObjectProperty = CastField<FObjectPropertyBase>(Property))
+	{
+		FString ObjectPath;
+		// TODO: No idea how we're meant to get object paths, might be this:
+		//ObjectProperty->ExportTextItem(ObjectPath, DataPtr, ??? /*DefaultValue*/, ??? /*Parent*/, 0, ??? /*ExportRootScope*/);
+		ResultValue = MakeShared<FJsonValueString>(ObjectPath);
+	}
+	else if (const FStructProperty* StructProperty = CastField<FStructProperty>(Property))
+	{
+		ResultValue = MakeShared<FJsonValueObject>(WriteStruct(StructProperty->Struct, DataPtr));
+	}
+	else if (const FArrayProperty* ArrayProperty = CastField<FArrayProperty>(Property))
+	{
+		FProperty* InnerProperty = ArrayProperty->Inner;
+		FScriptArrayHelper ArrayHelper(ArrayProperty, DataPtr);
+
+		TArray<TSharedPtr<FJsonValue>> JsonValuesArray;
+		JsonValuesArray.Reserve(ArrayHelper.Num());
+
+		for (int32 Idx = 0; Idx < ArrayHelper.Num(); ++Idx)
+		{
+			TSharedPtr<FJsonValue> ElementValue = WriteValue(InnerProperty, ArrayHelper.GetRawPtr(Idx));
+			JsonValuesArray.Add(ElementValue);
+		}
+
+		ResultValue = MakeShared<FJsonValueArray>(JsonValuesArray);
+	}
+	else if (const FSetProperty* SetProperty = CastField<FSetProperty>(Property))
+	{
+		FProperty* InnerProperty = SetProperty->ElementProp;
+		FScriptSetHelper SetHelper(SetProperty, DataPtr);
+
+		TArray<TSharedPtr<FJsonValue>> JsonValuesArray;
+		JsonValuesArray.Reserve(SetHelper.Num());
+
+		for (int32 Idx = 0; Idx < SetHelper.Num(); ++Idx)
+		{
+			if (SetHelper.IsValidIndex(Idx))
+			{
+				TSharedPtr<FJsonValue> ElementValue = WriteValue(InnerProperty, SetHelper.GetElementPtr(Idx));
+				JsonValuesArray.Add(ElementValue);
+			}
+		}
+
+		ResultValue = MakeShared<FJsonValueArray>(JsonValuesArray);
+	}
+	else if (const FMapProperty* MapProperty = CastField<FMapProperty>(Property))
+	{
+		FProperty* KeyProperty = MapProperty->KeyProp;
+		FProperty* ValueProperty = MapProperty->ValueProp;
+
+		FScriptMapHelper MapHelper(MapProperty, DataPtr);
+
+		if (MapHelper.Num() == 0)
+		{
+			ResultValue = MakeShared<FJsonValueObject>(MakeShared<FJsonObject>());
+		}
+		else
+		{
+			TArray<TSharedPtr<FJsonValue>> JsonKeysArray;
+			JsonKeysArray.Reserve(MapHelper.Num());
+
+			TArray<TSharedPtr<FJsonValue>> JsonValuesArray;
+			JsonValuesArray.Reserve(MapHelper.Num());
+
+			for (int32 Idx = 0; Idx < MapHelper.Num(); ++Idx)
+			{
+				TSharedPtr<FJsonValue> JsonKey = WriteValue(KeyProperty, MapHelper.GetKeyPtr(Idx));
+				JsonKeysArray.Add(JsonKey);
+
+				TSharedPtr<FJsonValue> JsonValue = WriteValue(ValueProperty, MapHelper.GetValuePtr(Idx));
+				JsonValuesArray.Add(JsonValue);
+			}
+
+			// maps can either be stored as $key, $value pairs or, if the keys can be stringified, as a JSON object
+			// check which we should use based on the first element
+			EJson KeyType = JsonKeysArray[0]->Type;
+			if (KeyType == EJson::Object)
+			{
+				TArray<TSharedPtr<FJsonValue>> ResultArray;
+				ResultArray.Reserve(MapHelper.Num());
+
+				for (int32 Idx = 0; Idx < MapHelper.Num(); ++Idx)
+				{
+					TSharedPtr<FJsonObject> ElementObject = MakeShared<FJsonObject>();
+					ElementObject->SetField(TEXT("$key"), JsonKeysArray[Idx]);
+					ElementObject->SetField(TEXT("$value"), JsonValuesArray[Idx]);
+
+					ResultArray.Add(MakeShared<FJsonValueObject>(ElementObject));
+				}
+
+				ResultValue = MakeShared<FJsonValueArray>(ResultArray);
+			}
+			else if (KeyType == EJson::Boolean || 
+				KeyType == EJson::Number || 
+				KeyType == EJson::String)
+			{
+				TSharedPtr<FJsonObject> ResultObject = MakeShared<FJsonObject>();
+
+				for (int32 Idx = 0; Idx < MapHelper.Num(); ++Idx)
+				{
+					FString KeyString;
+					check(JsonKeysArray[Idx]->TryGetString(KeyString));
+
+					ResultObject->SetField(KeyString, JsonValuesArray[Idx]);
+				}
+				
+				ResultValue = MakeShared<FJsonValueObject>(ResultObject);
+			}
+
+			ensureMsgf(ResultValue.IsValid(), TEXT("Map key type is invalid."));
+		}
+	}
+
+	ensureMsgf(ResultValue.IsValid(), TEXT("Property type is unsupported."));
+	return ResultValue;
+}
+
+bool FEditorConfig::IsDefault(const FProperty* Property, TSharedPtr<FJsonValue> JsonValue, const void* NativeValue)
+{
+	if (JsonValue->Type == EJson::Array)
+	{
+		return JsonValue->AsArray().Num() == 0;
+	}
+	else if (JsonValue->Type == EJson::Object)
+	{
+		return JsonValue->AsObject()->Values.Num() == 0;
+	}
+	else
+	{
+		// have the property initialize some temp storage, then check against that
+		uint8 Temp[256];
+		Property->InitializeValue(Temp);
+
+		return Property->Identical(NativeValue, Temp);
+	}
+
+	return false;
+}
+
+TSharedPtr<FJsonObject> FEditorConfig::WriteStruct(const UStruct* Struct, const void* Instance) 
+{
+	TSharedPtr<FJsonObject> JsonObject = MakeShared<FJsonObject>();
+
+	for (TFieldIterator<FProperty> It(Struct); It; ++It)
+	{
+		const FProperty* Property = *It;
+		const void* ValuePtr = Property->ContainerPtrToValuePtr<void>(Instance); 
+
+		TSharedPtr<FJsonValue> PropertyValue = WriteValue(Property, ValuePtr);
+		if (!IsDefault(Property, PropertyValue, ValuePtr))
+		{
+			JsonObject->SetField(Property->GetName(), PropertyValue);
+		}
+	}
+
+	return JsonObject;
+}
+
+/** 
+	* This exists because of sparse class data that can exist for UObjects only, and which is handled in ContainerPtrToValuePtr.
+	*/
+TSharedPtr<FJsonObject> FEditorConfig::WriteUObject(const UStruct* Struct, const UObject* Instance) 
+{
+	TSharedPtr<FJsonObject> JsonObject = MakeShared<FJsonObject>();
+
+	for (TFieldIterator<FProperty> It(Struct); It; ++It)
+	{
+		const FProperty* Property = *It;
+		const void* ValuePtr = Property->ContainerPtrToValuePtr<void>(Instance); 
+
+		TSharedPtr<FJsonValue> PropertyValue = WriteValue(Property, ValuePtr);
+		if (!IsDefault(Property, PropertyValue, ValuePtr))
+		{
+			JsonObject->SetField(Property->GetName(), PropertyValue);
+		}
+	}
+
+	return JsonObject;
+}
diff --git a/Engine/Source/Editor/EditorConfig/Private/EditorConfigModule.cpp b/Engine/Source/Editor/EditorConfig/Private/EditorConfigModule.cpp
new file mode 100644
--- /dev/null
+++ b/Engine/Source/Editor/EditorConfig/Private/EditorConfigModule.cpp
@@ -0,0 +1,17 @@
+// Copyright Epic Games, Inc. All Rights Reserved.
+
+#include "EditorConfigModule.h"
+#include "EditorConfig.h"
+#include "EditorConfigSubsystem.h"
+#include "EditorMetadataOverrides.h"
+#include "Misc/Paths.h"
+
+void FEditorConfigModule::StartupModule()
+{
+}
+
+void FEditorConfigModule::ShutdownModule()
+{
+}
+
+IMPLEMENT_MODULE(FEditorConfigModule, EditorConfig)
\ No newline at end of file
diff --git a/Engine/Source/Editor/EditorConfig/Private/JsonConfig.cpp b/Engine/Source/Editor/EditorConfig/Private/JsonConfig.cpp
new file mode 100644
--- /dev/null
+++ b/Engine/Source/Editor/EditorConfig/Private/JsonConfig.cpp
@@ -0,0 +1,1461 @@
+// Copyright Epic Games, Inc. All Rights Reserved.
+
+#include "JsonConfig.h"
+#include "Misc/FileHelper.h"
+#include "Serialization/JsonReader.h"
+#include "Serialization/JsonSerializer.h"
+
+namespace UE
+{
+	FJsonPath::FJsonPath()
+	{
+	}
+
+	FJsonPath::FJsonPath(const TCHAR* InPath)
+	{
+		ParsePath(InPath);
+	}
+
+	FJsonPath::FJsonPath(FStringView InPath)
+	{
+		ParsePath(FString(InPath));
+	}
+
+	FJsonPath::FJsonPath(const FJsonPath& Other) :
+		PathParts(Other.PathParts)
+	{
+	}
+
+	FJsonPath::FJsonPath(FJsonPath&& Other) :
+		PathParts(MoveTemp(Other.PathParts))
+	{
+	}
+
+	void FJsonPath::ParsePath(const FString& InPath)
+	{
+		TArray<FString> PathStrings;
+		InPath.ParseIntoArray(PathStrings, TEXT("."), false);
+
+		// we don't really care about the document root, remove it if it exists
+		if (PathStrings.Num() > 0 && PathStrings[0] == TEXT("$"))
+		{
+			PathStrings.RemoveAt(0);
+		}
+
+		for (const FString& Part : PathStrings)
+		{
+			if (Part.IsEmpty())
+			{
+				PathParts.Empty();
+				UE_LOG(LogEditorConfig, Warning, TEXT("Path part was empty in JSON path \"%s\""), *InPath);
+				return;
+			}
+
+			// check for index
+			int32 OpenBracket, CloseBracket;
+			if (Part.FindChar('[', OpenBracket))
+			{
+				if (Part.FindChar(']', CloseBracket) && 
+					CloseBracket > OpenBracket &&
+					CloseBracket == Part.Len() - 1)
+				{
+					FString Index = Part.Mid(OpenBracket + 1, (CloseBracket - OpenBracket) - 1);
+					if (Index.IsNumeric())
+					{
+						FPart NewPart;
+						NewPart.Name = *Part.Left(OpenBracket);
+						NewPart.Index = FCString::Atoi(*Index);
+					
+						PathParts.Add(NewPart);
+					}
+					else
+					{
+						PathParts.Empty();
+						UE_LOG(LogEditorConfig, Warning, TEXT("Path part \"%s\" looked to be index into array but was malformed in JSON path \"%s\""), *Part, *InPath);
+						return;
+					}
+				}
+				else
+				{
+					PathParts.Empty();
+					UE_LOG(LogEditorConfig, Warning, TEXT("Path part \"%s\" looked to be index into array but was malformed in JSON path \"%s\""), *Part, *InPath);
+					return;
+				}
+			}
+			else
+			{
+				FPart NewPart;
+				NewPart.Name = *Part;
+
+				// basic path
+				PathParts.Add(NewPart);
+			}
+		}
+	}
+
+	void FJsonPath::Append(FStringView Name)
+	{
+		FPart NewPart;
+		NewPart.Name = Name;
+		PathParts.Add(NewPart);
+	}
+
+	void FJsonPath::SetArrayIndex(int32 Index)
+	{
+		check(PathParts.Num() > 0);
+
+		PathParts.Last().Index = Index;
+	}
+
+	FJsonPath FJsonPath::GetSubPath(int32 NumParts) const
+	{
+		FJsonPath SubPath;
+		SubPath.PathParts.Reserve(NumParts);
+
+		for (int32 Idx = 0; Idx < NumParts; ++Idx)
+		{
+			SubPath.PathParts.Add(PathParts[Idx]);
+		}
+
+		return MoveTemp(SubPath);
+	}
+
+	FString FJsonPath::ToString() const
+	{
+		TStringBuilder<256> StringBuilder;
+		for (const FPart& Part : PathParts)
+		{
+			StringBuilder.Append(Part.Name);
+
+			if (Part.Index != INDEX_NONE)
+			{
+				StringBuilder.Append(TEXT("["));
+				StringBuilder.Append(LexToString(Part.Index));
+				StringBuilder.Append(TEXT("]"));
+			}
+		}
+
+		return StringBuilder.ToString();
+	}
+
+
+	//-----------------------------------------------------------------------------
+
+	FJsonConfig::FJsonConfig()
+	{
+		OverrideObject = MakeShared<FJsonObject>();
+		MergedObject = MakeShared<FJsonObject>();
+	}
+
+	void FJsonConfig::SetParent(TSharedPtr<FJsonConfig> Parent)
+	{
+		if (ParentConfig.IsValid())
+		{
+			ParentConfig->OnConfigChanged.Unbind();
+		}
+
+		ParentConfig = Parent;
+
+		MergeThisWithParent();
+
+		if (ParentConfig.IsValid())
+		{
+			ParentConfig->OnConfigChanged.BindRaw(this, &FJsonConfig::OnParentConfigChanged);
+		}
+
+		OnConfigChanged.ExecuteIfBound();
+	}
+
+	void FJsonConfig::OnParentConfigChanged()
+	{
+		MergeThisWithParent();
+	}
+
+	bool FJsonConfig::LoadFromFile(FStringView FilePath)
+	{
+		FString Contents;
+		if (!FFileHelper::LoadFileToString(Contents, FilePath.GetData()))
+		{
+			return false;
+		}
+
+		return LoadFromString(Contents);
+	}
+
+	bool FJsonConfig::LoadFromString(FStringView Content)
+	{
+		TSharedRef<FJsonStringReader> JsonReader = FJsonStringReader::Create(Content.GetData());
+		if (!FJsonSerializer::Deserialize(JsonReader.Get(), OverrideObject))
+		{
+			return false;
+		}
+
+		if (!OverrideObject.IsValid())
+		{
+			return false;
+		}
+
+		OnConfigChanged.ExecuteIfBound();
+
+		return MergeThisWithParent();
+	}
+
+	bool FJsonConfig::SaveToFile(FStringView FilePath) const
+	{
+		FString Contents;
+		if (!SaveToString(Contents))
+		{
+			return false;
+		}
+
+		return FFileHelper::SaveStringToFile(Contents, FilePath.GetData());
+	}
+
+	bool FJsonConfig::SaveToString(FString& OutResult) const
+	{
+		if (!OverrideObject.IsValid())
+		{
+			return false;
+		}
+
+		TSharedRef<TJsonWriter<>> Writer = TJsonWriterFactory<>::Create(&OutResult);
+		if (!FJsonSerializer::Serialize(OverrideObject.ToSharedRef(), Writer.Get()))
+		{
+			return false;
+		}
+
+		return true;
+	}
+
+	static TSharedPtr<FJsonValue> GetKeyField(const TSharedPtr<FJsonValue>& Value)
+	{
+		if (Value->Type == EJson::Object)
+		{
+			return Value->AsObject()->TryGetField(TEXT("$key"));
+		}
+
+		return TSharedPtr<FJsonValue>();
+	}
+
+	static TSharedPtr<FJsonValue> GetValueField(const TSharedPtr<FJsonValue>& Value)
+	{
+		if (Value->Type == EJson::Object)
+		{
+			return Value->AsObject()->TryGetField(TEXT("$value"));
+		}
+
+		return TSharedPtr<FJsonValue>();
+	}
+
+	static TSharedPtr<FJsonValue> FindValueInArray(const TArray<TSharedPtr<FJsonValue>>& Array, const TSharedPtr<FJsonValue>& Find)
+	{
+		if (Find->Type == EJson::Object)
+		{
+			TSharedPtr<FJsonValue> FindKey = GetKeyField(Find);
+			if (FindKey.IsValid())
+			{
+				const TSharedPtr<FJsonValue>* ExistingField = Array.FindByPredicate(
+					[FindKey](const TSharedPtr<FJsonValue>& Value)
+					{
+						TSharedPtr<FJsonValue> ValueKey = GetKeyField(Value);
+						if (ValueKey.IsValid())
+						{
+							return FJsonValue::CompareEqual(*FindKey.Get(), *ValueKey.Get());
+						}
+						return false;
+					});
+
+				if (ExistingField != nullptr)
+				{
+					return *ExistingField;
+				}
+			}
+		}
+		else
+		{
+			const TSharedPtr<FJsonValue>* ExistingField = Array.FindByPredicate(
+				[Find](const TSharedPtr<FJsonValue>& Value)
+				{
+					return FJsonValue::CompareEqual(*Find.Get(), *Value.Get());
+				});
+
+			if (ExistingField != nullptr)
+			{
+				return *ExistingField;
+			}
+		}
+
+		return TSharedPtr<FJsonValue>();
+	}
+
+	static bool RemoveValueFromArray(TArray<TSharedPtr<FJsonValue>>& Array, const TSharedPtr<FJsonValue>& Find)
+	{
+		if (Find->Type == EJson::Object)
+		{
+			TSharedPtr<FJsonValue> FindKey = GetKeyField(Find);
+			if (FindKey.IsValid())
+			{
+				return Array.RemoveAll(
+					[FindKey](const TSharedPtr<FJsonValue>& Value)
+					{
+						TSharedPtr<FJsonValue> ValueKey = GetKeyField(Value);
+						if (ValueKey.IsValid())
+						{
+							return FJsonValue::CompareEqual(*FindKey.Get(), *ValueKey.Get());
+						}
+						return false;
+					}) > 0;
+			}
+		}
+
+		return Array.RemoveAll(
+			[Find](const TSharedPtr<FJsonValue>& Value)
+			{
+				return FJsonValue::CompareEqual(*Find.Get(), *Value.Get());
+			}) > 0;
+	}
+
+	static bool ApplyDeltaOperationsToArray(const TSharedPtr<FJsonObject>& OverrideObject, TArray<TSharedPtr<FJsonValue>>& DestArray)
+	{
+		const TArray<TSharedPtr<FJsonValue>>* OverrideSet = nullptr;
+		if (OverrideObject->TryGetArrayField(TEXT("="), OverrideSet))
+		{
+			DestArray.Reset();
+			DestArray.Reserve(OverrideSet->Num());
+			for (const TSharedPtr<FJsonValue>& SetValue : *OverrideSet)
+			{
+				DestArray.Add(FJsonValue::Duplicate(SetValue));
+			}
+		}
+
+		const TArray<TSharedPtr<FJsonValue>>* OverrideAdd = nullptr;
+		if (OverrideObject->TryGetArrayField(TEXT("+"), OverrideAdd))
+		{
+			DestArray.Reserve(DestArray.Num() + OverrideAdd->Num());
+			for (const TSharedPtr<FJsonValue>& AddValue : *OverrideAdd)
+			{
+				// check if this is a map with $key, $value fields
+				if (AddValue->Type == EJson::Object)
+				{
+					TSharedPtr<FJsonValue> AddKey = GetKeyField(AddValue);
+					if (AddKey.IsValid())
+					{
+						TSharedPtr<FJsonValue>* ExistingField = DestArray.FindByPredicate(
+							[AddKey](const TSharedPtr<FJsonValue>& Value)
+							{
+								TSharedPtr<FJsonValue> ValueKey = GetKeyField(Value);
+								if (ValueKey.IsValid())
+								{
+									return FJsonValue::CompareEqual(*ValueKey.Get(), *AddKey.Get());
+								}
+								return false;
+							});
+
+						if (ExistingField != nullptr)
+						{
+							// there was an existing field with the same key, override that value
+							TSharedPtr<FJsonValue> AddValueField = GetValueField(AddValue);
+							ensureAlways(AddValueField.IsValid());
+
+							TSharedPtr<FJsonObject> ExistingObject = (*ExistingField)->AsObject();
+							ExistingObject->SetField(TEXT("$value"), FJsonValue::Duplicate(AddValueField));
+							continue;
+						}
+					}
+				}
+
+				DestArray.Add(FJsonValue::Duplicate(AddValue));
+			}
+		}
+
+		const TArray<TSharedPtr<FJsonValue>>* OverrideRemove = nullptr;
+		if (OverrideObject->TryGetArrayField(TEXT("-"), OverrideRemove))
+		{
+			for (const TSharedPtr<FJsonValue>& RemoveValue : *OverrideRemove)
+			{
+				DestArray.RemoveAll(
+					[RemoveValue](const TSharedPtr<FJsonValue>& Value)
+					{
+						TSharedPtr<FJsonValue> RemoveKey = GetKeyField(RemoveValue);
+						TSharedPtr<FJsonValue> ExistingKey = GetKeyField(Value);
+						if (RemoveKey.IsValid() && ExistingKey.IsValid())
+						{
+							// the existing and remove values have $key fields, compare those
+							return FJsonValue::CompareEqual(*RemoveKey.Get(), *ExistingKey.Get());
+						}
+
+						if (ExistingKey.IsValid())
+						{
+							// the existing value has a $key field, compare it with the remove value, 
+							// because we want to support the short-hand of "-": [ "Foo" ] instead of "-": [ { "$key": "Foo" } ]
+							return FJsonValue::CompareEqual(*RemoveValue.Get(), *ExistingKey.Get());
+						}
+
+						// just compare values directly
+						return FJsonValue::CompareEqual(*RemoveValue.Get(), *Value.Get());
+					});
+			}
+		}
+	
+		if (OverrideSet != nullptr && 
+			(OverrideAdd != nullptr || OverrideRemove != nullptr))
+		{
+			UE_LOG(LogEditorConfig, Warning, TEXT("JSON container is malformed and contains both set and add/remove values."));
+		}
+
+		return OverrideSet != nullptr || OverrideAdd != nullptr || OverrideRemove != nullptr;
+	}
+
+	static bool ApplyOverridesToObject(const TSharedPtr<FJsonObject>& Overrides, TSharedPtr<FJsonObject>& Dest)
+	{
+		for (const TPair<FString, TSharedPtr<FJsonValue>>& Override : Overrides->Values)
+		{
+			TSharedPtr<FJsonValue> DestValue = Dest->TryGetField(Override.Key);
+			if (!DestValue.IsValid())
+			{
+				const TSharedPtr<FJsonObject>* OverrideObject;
+				if (Override.Value->TryGetObject(OverrideObject))
+				{
+					TArray<TSharedPtr<FJsonValue>> DestArray;
+					if (ApplyDeltaOperationsToArray(*OverrideObject, DestArray))
+					{
+						// was an array, DestArray is populated
+						Dest->SetArrayField(Override.Key, DestArray);
+					}
+					else
+					{
+						TSharedPtr<FJsonObject> DestObject = MakeShared<FJsonObject>();
+						ApplyOverridesToObject(*OverrideObject, DestObject);
+						Dest->SetObjectField(Override.Key, DestObject);
+					}
+				}
+				else
+				{
+					Dest->SetField(Override.Key, FJsonValue::Duplicate(Override.Value));
+				}
+			}
+			else
+			{
+				switch (Override.Value->Type)
+				{
+					case EJson::Boolean:
+					{
+						bool OverrideBool, DestBool;
+						if (Override.Value->TryGetBool(OverrideBool) && 
+							DestValue->TryGetBool(DestBool))
+						{
+							Dest->SetBoolField(Override.Key, OverrideBool);
+						}
+						continue;
+					}
+					case EJson::Number:
+					{
+						double OverrideNumber, DestNumber;
+						if (Override.Value->TryGetNumber(OverrideNumber) &&
+							DestValue->TryGetNumber(DestNumber))
+						{
+							Dest->SetNumberField(Override.Key, OverrideNumber);
+						}
+						continue;
+					}
+					case EJson::String:
+					{
+						FString OverrideString, DestString;
+						if (Override.Value->TryGetString(OverrideString) &&
+							DestValue->TryGetString(DestString))
+						{
+							Dest->SetStringField(Override.Key, OverrideString);
+						}
+						continue;
+					}
+					case EJson::Array:
+					{
+						if (Override.Value->Type == EJson::Array && 
+							DestValue->Type == EJson::Array)
+						{
+							// our overrides has an array, so this is implicitly a set operation
+							Dest->SetField(Override.Key, FJsonValue::Duplicate(Override.Value));
+						}
+						continue;
+					}
+					case EJson::Object:
+					{
+						const TSharedPtr<FJsonObject>* OverrideObject;
+						if (Override.Value->TryGetObject(OverrideObject))
+						{
+							TArray<TSharedPtr<FJsonValue>>* DestArray = nullptr;
+							TSharedPtr<FJsonObject>* DestObject = nullptr;
+							if (DestValue->TryGetObject(DestObject))
+							{
+								ApplyOverridesToObject(*OverrideObject, *DestObject);
+							}
+							else if (DestValue->TryGetArray(DestArray))
+							{
+								ApplyDeltaOperationsToArray(*OverrideObject, *DestArray);
+							}
+						}
+						continue;
+					}
+					default:
+						ensureAlways(false);
+						return false;
+				}
+			}
+		}
+
+		return true;
+	}
+
+	bool FJsonConfig::MergeThisWithParent()
+	{
+		MergedObject->Values.Reset();
+
+		// pre-fill with parent's values if we have one
+		if (ParentConfig.IsValid())
+		{
+			FJsonObject::Duplicate(ParentConfig->MergedObject, MergedObject);
+		}
+
+		// apply this object's overrides if we have any data
+		if (OverrideObject.IsValid())
+		{
+			return ApplyOverridesToObject(OverrideObject, MergedObject);
+		}
+
+		return true;
+	}
+
+	bool FJsonConfig::TryGetBool(const FJsonPath& Path, bool& OutValue) const
+	{
+		TSharedPtr<FJsonValue> JsonValue;
+		if (!TryGetJsonValue(Path, JsonValue))
+		{
+			return false;
+		}
+	
+		return JsonValue->TryGetBool(OutValue);
+	}
+
+	bool FJsonConfig::TryGetString(const FJsonPath& Path, FString& OutValue) const
+	{
+		TSharedPtr<FJsonValue> JsonValue;
+		if (!TryGetJsonValue(Path, JsonValue))
+		{
+			return false;
+		}
+	
+		return JsonValue->TryGetString(OutValue);
+	}
+
+	bool FJsonConfig::TryGetString(const FJsonPath& Path, FName& OutValue) const
+	{
+		TSharedPtr<FJsonValue> JsonValue;
+		if (!TryGetJsonValue(Path, JsonValue))
+		{
+			return false;
+		}
+
+		FString TempString;
+		if (!JsonValue->TryGetString(TempString))
+		{
+			return false;
+		}
+
+		OutValue = *TempString;
+		return true;
+	}
+
+	bool FJsonConfig::TryGetString(const FJsonPath& Path, FText& OutValue) const
+	{
+		TSharedPtr<FJsonValue> JsonValue;
+		if (!TryGetJsonValue(Path, JsonValue))
+		{
+			return false;
+		}
+	
+		FString TempString;
+		if (!JsonValue->TryGetString(TempString))
+		{
+			return false;
+		}
+
+		OutValue = FText::FromString(TempString);
+		return true;
+	}
+
+	bool FJsonConfig::TryGetArray(const FJsonPath& Path, TArray<bool>& OutArray) const
+	{
+		return TryGetArrayHelper(Path, OutArray, 
+			[](const TSharedPtr<FJsonValue>& JsonValue, bool& OutBool)
+			{
+				return JsonValue->TryGetBool(OutBool);
+			});
+	}
+
+	bool FJsonConfig::TryGetArray(const FJsonPath& Path, TArray<int8>& OutArray) const
+	{
+		return TryGetNumericArrayHelper(Path, OutArray);
+	}
+
+	bool FJsonConfig::TryGetArray(const FJsonPath& Path, TArray<int16>& OutArray) const
+	{
+		return TryGetNumericArrayHelper(Path, OutArray);
+	}
+
+	bool FJsonConfig::TryGetArray(const FJsonPath& Path, TArray<int32>& OutArray) const
+	{
+		return TryGetNumericArrayHelper(Path, OutArray);
+	}
+
+	bool FJsonConfig::TryGetArray(const FJsonPath& Path, TArray<int64>& OutArray) const
+	{
+		return TryGetNumericArrayHelper(Path, OutArray);
+	}
+
+	bool FJsonConfig::TryGetArray(const FJsonPath& Path, TArray<uint8>& OutArray) const
+	{
+		return TryGetNumericArrayHelper(Path, OutArray);
+	}
+
+	bool FJsonConfig::TryGetArray(const FJsonPath& Path, TArray<uint16>& OutArray) const
+	{
+		return TryGetNumericArrayHelper(Path, OutArray);
+	}
+
+	bool FJsonConfig::TryGetArray(const FJsonPath& Path, TArray<uint32>& OutArray) const
+	{
+		return TryGetNumericArrayHelper(Path, OutArray);
+	}
+
+	bool FJsonConfig::TryGetArray(const FJsonPath& Path, TArray<uint64>& OutArray) const
+	{
+		return TryGetNumericArrayHelper(Path, OutArray);
+	}
+
+	bool FJsonConfig::TryGetArray(const FJsonPath& Path, TArray<float>& OutArray) const
+	{
+		return TryGetNumericArrayHelper(Path, OutArray);
+	}
+
+	bool FJsonConfig::TryGetArray(const FJsonPath& Path, TArray<double>& OutArray) const
+	{
+		return TryGetNumericArrayHelper(Path, OutArray);
+	}
+
+	bool FJsonConfig::TryGetArray(const FJsonPath& Path, TArray<FString>& OutArray) const
+	{
+		return TryGetArrayHelper(Path, OutArray, 
+			[](const TSharedPtr<FJsonValue>& JsonValue, FString& OutString)
+			{
+				return JsonValue->TryGetString(OutString);
+			});
+	}
+
+	bool FJsonConfig::TryGetArray(const FJsonPath& Path, TArray<FText>& OutArray) const
+	{
+		return TryGetArrayHelper(Path, OutArray, 
+			[](const TSharedPtr<FJsonValue>& JsonValue, FText& OutText)
+			{
+				FString StringValue;
+				if (JsonValue->TryGetString(StringValue))
+				{
+					OutText = FText::FromString(StringValue);
+					return true;
+				}
+				return false;
+			});
+	}
+
+	bool FJsonConfig::TryGetArray(const FJsonPath& Path, TArray<FName>& OutArray) const
+	{
+		return TryGetArrayHelper(Path, OutArray, 
+			[](const TSharedPtr<FJsonValue>& JsonValue, FName& OutName)
+			{
+				FString StringValue;
+				if (JsonValue->TryGetString(StringValue))
+				{
+					OutName = *StringValue;
+					return true;
+				}
+				return false;
+			});
+	}
+
+	bool FJsonConfig::TryGetArray(const FJsonPath& Path, TArray<TSharedPtr<FJsonValue>>& OutValue) const
+	{
+		TSharedPtr<FJsonValue> JsonValue;
+		if (!TryGetJsonValue(Path, JsonValue))
+		{
+			return false;
+		}
+
+		const TArray<TSharedPtr<FJsonValue>>* ValuePtr;
+		if (!JsonValue->TryGetArray(ValuePtr))
+		{
+			return false;
+		}
+
+		OutValue = *ValuePtr;
+		return true;
+	}
+
+	bool FJsonConfig::TryGetMap(const FJsonPath& Path, TArray<FJsonValuePair>& OutMap) const
+	{
+		TArray<TSharedPtr<FJsonValue>> Array;
+		if (!TryGetArray(Path, Array))
+		{
+			return false;
+		}
+
+		OutMap.Reset();
+		OutMap.Reserve(Array.Num());
+
+		for (const TSharedPtr<FJsonValue>& Value : Array)
+		{
+			if (Value->Type != EJson::Object)
+			{
+				return false;
+			}
+
+			const TSharedPtr<FJsonObject>& Object = Value->AsObject();
+
+			TSharedPtr<FJsonValue> KeyField = Object->TryGetField(TEXT("$key"));
+			TSharedPtr<FJsonValue> ValueField = Object->TryGetField(TEXT("$value"));
+
+			if (!KeyField.IsValid() || !ValueField.IsValid())
+			{
+				return false;
+			}
+
+			OutMap.Add(FJsonValuePair(KeyField, ValueField));
+		}
+
+		return true;
+	}
+
+	bool FJsonConfig::TryGetJsonObject(const FJsonPath& Path, TSharedPtr<FJsonObject>& OutValue) const
+	{
+		TSharedPtr<FJsonValue> JsonValue;
+		if (!TryGetJsonValue(Path, JsonValue))
+		{
+			return false;
+		}
+
+		const TSharedPtr<FJsonObject>* ValuePtr;
+		if (JsonValue->TryGetObject(ValuePtr))
+		{
+			OutValue = *ValuePtr;
+			return true;
+		}
+		return false;
+	}
+
+	TSharedPtr<FJsonObject> FJsonConfig::GetRootObject() const
+	{
+		return MergedObject;
+	}
+
+	static const TSharedPtr<FJsonValue>* FindPartInObject(const TSharedPtr<FJsonObject>& Object, const FJsonPath::FPart& PathPart)
+	{
+		const TSharedPtr<FJsonValue>* NextValue = Object->Values.Find(PathPart.Name);
+		if (NextValue == nullptr)
+		{
+			// part of path didn't exist
+			return nullptr;
+		}
+
+		if (PathPart.Index != INDEX_NONE)
+		{
+			const TArray<TSharedPtr<FJsonValue>>* ArrayPtr = nullptr;
+			if ((*NextValue)->TryGetArray(ArrayPtr))
+			{
+				if (!ArrayPtr->IsValidIndex(PathPart.Index))
+				{
+					// index out of range
+					return nullptr;
+				}
+
+				NextValue = &(*ArrayPtr)[PathPart.Index];
+			}
+		}
+		
+		return NextValue;
+	}
+
+	static TSharedPtr<FJsonValue>* FindPartInObject(TSharedPtr<FJsonObject>& Object, const FJsonPath::FPart& PathPart)
+	{
+		return const_cast<TSharedPtr<FJsonValue>*>(FindPartInObject((const TSharedPtr<FJsonObject>&) Object, PathPart));
+	}
+
+	bool FJsonConfig::TryGetJsonValue(const FJsonPath& Path, TSharedPtr<FJsonValue>& OutValue) const
+	{
+		if (!IsValid())
+		{
+			return false;
+		}
+
+		const TSharedPtr<FJsonObject>* CurrentObject = &MergedObject;
+
+		for (int32 Idx = 0; Idx < Path.Length() - 1; ++Idx)
+		{
+			const TSharedPtr<FJsonValue>* NextValue = FindPartInObject(*CurrentObject, Path[Idx]);
+			if (NextValue == nullptr || !(*NextValue)->TryGetObject(CurrentObject))
+			{
+				return false;
+			}
+		}
+
+		const TSharedPtr<FJsonValue>* LastValue = FindPartInObject(*CurrentObject, Path[Path.Length() - 1]);
+		if (LastValue == nullptr)
+		{
+			return false;
+		}
+
+		OutValue = *LastValue;
+		return true;
+	}
+
+	static bool SetValueHelper(TSharedPtr<FJsonValue>& CurrentValue, TSharedPtr<FJsonValue> NewValue, TSharedPtr<FJsonValue> ParentValue);
+
+	static bool SetArrayValueHelper(TSharedPtr<FJsonValue> CurrentValue, const TArray<TSharedPtr<FJsonValue>>& NewArray, TSharedPtr<FJsonValue> ParentValue)
+	{
+		TArray<TSharedPtr<FJsonValue>>* CurrentArray = nullptr;
+		if (CurrentValue->TryGetArray(CurrentArray))
+		{
+			// currently has array value so is implicitly a set, we can just overwrite it
+			(*CurrentArray) = NewArray;
+			return true;
+		}
+
+		TSharedPtr<FJsonObject>* CurrentObject = nullptr;
+		if (!CurrentValue->TryGetObject(CurrentObject))
+		{
+			return false;
+		}
+
+		// we generally want to steer clear of set fields, but if one already exists then we'll just use it
+		TSharedPtr<FJsonValue> SetField = (*CurrentObject)->TryGetField(TEXT("="));
+		TArray<TSharedPtr<FJsonValue>>* SetArray = nullptr;
+		if (SetField.IsValid() && SetField->TryGetArray(SetArray))
+		{
+			(*SetArray) = NewArray;
+
+			// remove add and remove fields
+			(*CurrentObject)->RemoveField(TEXT("+"));
+			(*CurrentObject)->RemoveField(TEXT("-"));
+			return true;
+		}
+
+		// build up a diff relative to the parent array
+		const TArray<TSharedPtr<FJsonValue>>* ParentArray = nullptr;
+		if (ParentValue.IsValid())
+		{
+			ParentValue->TryGetArray(ParentArray);
+		}
+
+		TArray<TSharedPtr<FJsonValue>> AddedValues;
+		TArray<TSharedPtr<FJsonValue>> RemovedValues;
+
+		if (ParentArray != nullptr)
+		{
+			// find all shared values
+			TArray<TSharedPtr<FJsonValue>> SharedValues;
+			SharedValues.Reserve(NewArray.Num());
+
+			for (const TSharedPtr<FJsonValue>& NewElement : NewArray)
+			{
+				TSharedPtr<FJsonValue> ExistingElement = FindValueInArray(*ParentArray, NewElement);
+				if (ExistingElement.IsValid())
+				{
+					SharedValues.Add(NewElement);
+				}
+			}
+
+			// add all values in the new array that weren't shared to the added values
+			AddedValues.Reserve(NewArray.Num() - SharedValues.Num());
+
+			for (const TSharedPtr<FJsonValue>& NewElement : NewArray)
+			{
+				TSharedPtr<FJsonValue> ExistingElement = FindValueInArray(SharedValues, NewElement);
+				if (!ExistingElement.IsValid())
+				{
+					AddedValues.Add(NewElement);
+				}
+			}
+	
+			// add all values in the parent array that weren't shared to the removed values
+			RemovedValues.Reserve(ParentArray->Num() - SharedValues.Num());
+
+			for (const TSharedPtr<FJsonValue>& ParentElement : *ParentArray)
+			{
+				TSharedPtr<FJsonValue> ExistingElement = FindValueInArray(SharedValues, ParentElement);
+				if (!ExistingElement.IsValid())
+				{
+					RemovedValues.Add(ParentElement);
+				}
+			}
+		}
+		else
+		{
+			AddedValues = NewArray;
+		}
+
+		TSharedPtr<FJsonValue> AddField = (*CurrentObject)->TryGetField(TEXT("+"));
+		if (AddField.IsValid())
+		{
+			TArray<TSharedPtr<FJsonValue>>* AddArray = nullptr;
+			if (ensureAlwaysMsgf(AddField->TryGetArray(AddArray), TEXT("Invalid JSON config: \"+\" field in JSON config was not an array.")))
+			{
+				(*AddArray) = AddedValues;
+			}
+		}
+		else
+		{
+			(*CurrentObject)->SetArrayField(TEXT("+"), AddedValues);
+		}
+
+		TSharedPtr<FJsonValue> RemoveField = (*CurrentObject)->TryGetField(TEXT("-"));
+		if (RemoveField.IsValid())
+		{
+			TArray<TSharedPtr<FJsonValue>>* RemoveArray = nullptr;
+			if (ensureAlwaysMsgf(RemoveField->TryGetArray(RemoveArray), TEXT("Invalid JSON config: \"-\" field in JSON config was not an array.")))
+			{
+				(*RemoveArray) = RemovedValues;
+			}
+		}
+		else
+		{
+			(*CurrentObject)->SetArrayField(TEXT("-"), RemovedValues);
+		}
+
+		return true;
+	}
+
+	static bool SetObjectValueHelper(TSharedPtr<FJsonObject> CurrentObject, TSharedPtr<FJsonObject> NewObject, TSharedPtr<FJsonValue> ParentValue)
+	{
+		TSharedPtr<FJsonObject>* ParentObject = nullptr;
+		if (ParentValue.IsValid())
+		{
+			ParentValue->TryGetObject(ParentObject);
+		}
+
+		if (ParentObject == nullptr)
+		{
+			// no parent, just set everything we were given into the object
+			for (const TPair<FString, TSharedPtr<FJsonValue>>& NewValuePair : NewObject->Values)
+			{
+				CurrentObject->SetField(NewValuePair.Key, NewValuePair.Value);
+			}
+		}
+		else
+		{
+			for (const TPair<FString, TSharedPtr<FJsonValue>>& NewValuePair : NewObject->Values)
+			{
+				TSharedPtr<FJsonValue> ParentField = (*ParentObject)->TryGetField(NewValuePair.Key);
+				if (!ParentField.IsValid())
+				{
+					// not set in parent, just set here
+					CurrentObject->SetField(NewValuePair.Key, NewValuePair.Value);
+				}
+				else
+				{
+					// already set in parent, if we're identical we can just remove this field from this object
+					if (FJsonValue::CompareEqual(*ParentField.Get(), *NewValuePair.Value.Get()))
+					{
+						CurrentObject->RemoveField(NewValuePair.Key);
+					}
+					else
+					{
+						TSharedPtr<FJsonValue> CurrentField = CurrentObject->TryGetField(NewValuePair.Key);
+						if (!CurrentField.IsValid())
+						{
+							CurrentObject->SetField(NewValuePair.Key, NewValuePair.Value);
+						}
+						else
+						{
+							SetValueHelper(CurrentField, NewValuePair.Value, ParentField);
+						}
+					}
+				}
+			}
+		}
+
+		return true;
+	}
+
+	bool SetValueHelper(TSharedPtr<FJsonValue>& CurrentValue, TSharedPtr<FJsonValue> NewValue, TSharedPtr<FJsonValue> ParentValue)
+	{
+		check(CurrentValue.IsValid());
+		check(NewValue.IsValid());
+
+		EJson ValueType = NewValue->Type;
+
+		switch (ValueType)
+		{
+			case EJson::Number:
+			case EJson::Boolean:
+			case EJson::String:
+				if (CurrentValue->Type != ValueType)
+				{
+					return false;
+				}
+
+				CurrentValue = NewValue;
+				return true;
+
+			case EJson::Array:
+				if (CurrentValue->Type != EJson::Object || CurrentValue->Type != EJson::Array)
+				{
+					return false;
+				}
+
+				return SetArrayValueHelper(CurrentValue, NewValue->AsArray(), ParentValue);
+
+			case EJson::Object:
+				if (CurrentValue->Type != EJson::Object)
+				{
+					return false;
+				}
+
+				return SetObjectValueHelper(CurrentValue->AsObject(), NewValue->AsObject(), ParentValue);
+		}
+
+		return false;
+	}
+
+	bool FJsonConfig::SetJsonValueInMerged(const FJsonPath& Path, TSharedPtr<FJsonValue> NewValue)
+	{
+		TSharedPtr<FJsonObject>* CurrentObject = &MergedObject;
+	
+		for (int32 Idx = 0; Idx < Path.Length() - 1; ++Idx)
+		{
+			const FJsonPath::FPart& CurrentPart = Path[Idx];
+			const TSharedPtr<FJsonValue>* NextValue = FindPartInObject(*CurrentObject, CurrentPart);
+			if (NextValue == nullptr || 
+				!(*NextValue)->TryGetObject(CurrentObject))
+			{
+				if (CurrentPart.Index == INDEX_NONE)
+				{
+					// create an empty object
+					(*CurrentObject)->SetObjectField(Path[Idx].Name, MakeShared<FJsonObject>());
+				}
+				else
+				{
+					// array entry that doesn't exist
+					return false;
+				}
+			}
+		}
+
+		const FJsonPath::FPart& LastPart = Path[Path.Length() - 1];
+
+		TSharedPtr<FJsonValue>* LastValue = FindPartInObject(*CurrentObject, LastPart);
+		if (LastValue == nullptr)
+		{
+			if (LastPart.Index == INDEX_NONE)
+			{
+				// value doesn't exist and this isn't an array entry, should just set it
+				(*CurrentObject)->SetField(LastPart.Name, NewValue);
+				return true;
+			}
+			else
+			{
+				// array entry that doesn't exist
+				return false;
+			}
+		}
+
+		return SetValueHelper(*LastValue, NewValue, TSharedPtr<FJsonValue>());
+	}
+
+	bool FJsonConfig::SetJsonValueInOverride(const FJsonPath& Path, TSharedPtr<FJsonValue> NewValue, TSharedPtr<FJsonValue> PreviousValue, TSharedPtr<FJsonValue> ParentValue)
+	{
+		TSharedPtr<FJsonObject>* CurrentObject = &OverrideObject;
+	
+		for (int32 Idx = 0; Idx < Path.Length() - 1; ++Idx)
+		{
+			const FJsonPath::FPart& CurrentPart = Path[Idx];
+			const TSharedPtr<FJsonValue>* NextValue = FindPartInObject(*CurrentObject, CurrentPart);
+			if (NextValue == nullptr)
+			{
+				// since we know this isn't the last part, we can safely add an object here
+				(*CurrentObject)->SetObjectField(CurrentPart.Name, MakeShared<FJsonObject>());
+
+				TSharedPtr<FJsonValue> CreatedValue = (*CurrentObject)->TryGetField(CurrentPart.Name);
+				CreatedValue->TryGetObject(CurrentObject);
+			}
+			else if (!(*NextValue)->TryGetObject(CurrentObject))
+			{
+				FString SubPath = Path.GetSubPath(Idx + 1).ToString();
+				ensureMsgf(false, TEXT("JSON value at \"%s\" not a JSON object in path: %s"), *SubPath, *Path.ToString());
+			}
+
+			if (CurrentPart.Index != INDEX_NONE)
+			{
+				// if an index is set, there's an array in the path, eg. Foo[1].Bar
+				ensureAlwaysMsgf(false, TEXT("The array in the middle of a path needs to be figured out..."));
+			}
+		}
+
+		const FJsonPath::FPart& LastPart = Path[Path.Length() - 1];
+
+		TSharedPtr<FJsonValue>* LastValue = FindPartInObject(*CurrentObject, LastPart);
+		if (LastValue == nullptr)
+		{
+			// field doesn't exist 
+
+			if (LastPart.Index == INDEX_NONE)
+			{
+				// this isn't an array entry, just set it now
+				if (NewValue->Type == EJson::Object || NewValue->Type == EJson::Array)
+				{
+					// create a new empty object that we can diff against
+					TSharedPtr<FJsonValue> CreatedValue = MakeShared<FJsonValueObject>(MakeShared<FJsonObject>());
+					(*CurrentObject)->SetField(LastPart.Name, CreatedValue);
+					SetValueHelper(CreatedValue, NewValue, ParentValue);
+				}
+				else
+				{
+					(*CurrentObject)->SetField(LastPart.Name, NewValue);
+				}
+
+				return true;
+			}
+			else if (!PreviousValue.IsValid())
+			{
+				// out of bounds
+				return false;
+			}
+			else
+			{
+				// valid index, but no array field created in OverrideObject yet, create one
+				(*CurrentObject)->SetObjectField(LastPart.Name, MakeShared<FJsonObject>());
+			}
+		}
+
+		if (LastPart.Index != INDEX_NONE)
+		{
+			// last part is an array, eg. Foo.Bar[0] - array should absolutely be created by now
+			TSharedPtr<FJsonValue> ArrayField = (*CurrentObject)->TryGetField(LastPart.Name);
+			check(ArrayField.IsValid());
+
+			TArray<TSharedPtr<FJsonValue>>* ArrayValue = nullptr;
+			if (ArrayField->TryGetArray(ArrayValue))
+			{
+				// override is an array, so implicitly a set
+
+				if (ArrayValue->IsValidIndex(LastPart.Index))
+				{
+					(*ArrayValue)[LastPart.Index] = NewValue;
+					return true;
+				}
+
+				// index out of bounds
+				return false;
+			}
+
+			TSharedPtr<FJsonObject>* ArrayObject = nullptr;
+			if (ArrayField->TryGetObject(ArrayObject))
+			{
+				// the field is an object, so we need to modify existing override fields
+
+				// check if we have a set field
+				TSharedPtr<FJsonValue> SetField = (*ArrayObject)->TryGetField(TEXT("="));
+				TArray<TSharedPtr<FJsonValue>>* SetContents = nullptr;
+				if (SetField.IsValid() && SetField->TryGetArray(SetContents))
+				{
+					if (SetContents->IsValidIndex(LastPart.Index))
+					{
+						(*SetContents)[LastPart.Index] = NewValue;
+						return true;
+					}
+
+					// index out of bounds
+					return false;
+				}
+
+				// see if the previous value has already been removed in an override
+				check(PreviousValue.IsValid());
+
+				TSharedPtr<FJsonValue> RemoveField = (*ArrayObject)->TryGetField(TEXT("-"));
+				TArray<TSharedPtr<FJsonValue>>* RemoveContents = nullptr;
+				if (RemoveField.IsValid() && RemoveField->TryGetArray(RemoveContents))
+				{
+					TSharedPtr<FJsonValue> ExistingRemove = FindValueInArray(*RemoveContents, PreviousValue);
+					if (!ExistingRemove.IsValid())
+					{
+						// hasn't already been removed, remove it now
+						RemoveContents->Add(PreviousValue);
+					}
+				
+					// remove the new value from the remove list, otherwise we can end up in a situation where a value is added _and_ removed
+					RemoveValueFromArray(*RemoveContents, NewValue);
+				}
+				else
+				{
+					// add a remove field with the previous value
+					TArray<TSharedPtr<FJsonValue>> NewRemoveContents; 
+					NewRemoveContents.Add(PreviousValue);
+
+					(*ArrayObject)->SetArrayField(TEXT("-"), NewRemoveContents);
+				}
+
+				// see if the new value has already been added in an override
+				TSharedPtr<FJsonValue> AddField = (*ArrayObject)->TryGetField(TEXT("+"));
+				TArray<TSharedPtr<FJsonValue>>* AddContents = nullptr;
+				if (AddField.IsValid() && AddField->TryGetArray(AddContents))
+				{
+					TSharedPtr<FJsonValue> ExistingAdd = FindValueInArray(*AddContents, NewValue);
+					if (!ExistingAdd.IsValid())
+					{
+						// hasn't already been added, add it now
+						AddContents->Add(NewValue);
+					}
+
+					// remove the old value from the add list, otherwise we can end up in a situation where a value is added _and_ removed
+					RemoveValueFromArray(*AddContents, PreviousValue);
+				}
+				else
+				{
+					// add an add field with the new value
+					TArray<TSharedPtr<FJsonValue>> NewAddContents; 
+					NewAddContents.Add(NewValue);
+
+					(*ArrayObject)->SetArrayField(TEXT("+"), NewAddContents);
+				}
+			}
+
+			return true;
+		}
+	
+		return SetValueHelper(*LastValue, NewValue, ParentValue);
+	}
+
+	bool FJsonConfig::RemoveJsonValueFromOverride(const FJsonPath& Path, TSharedPtr<FJsonValue> CurrentValue)
+	{
+		TSharedPtr<FJsonObject>* CurrentObject = &OverrideObject;
+	
+		for (int32 Idx = 0; Idx < Path.Length() - 1; ++Idx)
+		{
+			const TSharedPtr<FJsonValue>* NextValue = FindPartInObject(*CurrentObject, Path[Idx]);
+			if (NextValue == nullptr || !(*NextValue)->TryGetObject(CurrentObject))
+			{
+				return false;
+			}
+		}
+
+		const FJsonPath::FPart& LastPart = Path[Path.Length() - 1];
+
+		if (LastPart.Index != INDEX_NONE)
+		{
+			TSharedPtr<FJsonValue> ArrayField = (*CurrentObject)->TryGetField(LastPart.Name);
+			if (!ArrayField.IsValid())
+			{
+				return false;
+			}
+
+			TArray<TSharedPtr<FJsonValue>>* ArrayPtr = nullptr;
+			if (ArrayField->TryGetArray(ArrayPtr))
+			{
+				// array field is an array, so implicitly a set
+				RemoveValueFromArray(*ArrayPtr, CurrentValue);
+				return true;
+			}
+
+			TSharedPtr<FJsonObject>* ArrayObject = nullptr;
+			if (ArrayField->TryGetObject(ArrayObject))
+			{
+				{
+					TSharedPtr<FJsonValue> SetField = (*ArrayObject)->TryGetField(TEXT("="));
+					TArray<TSharedPtr<FJsonValue>>* SetContents = nullptr;
+					if (SetField.IsValid() && SetField->TryGetArray(SetContents))
+					{
+						RemoveValueFromArray(*SetContents, CurrentValue);
+						return true;
+					}
+				}
+
+				{
+					TSharedPtr<FJsonValue> AddField = (*ArrayObject)->TryGetField(TEXT("+"));
+					TArray<TSharedPtr<FJsonValue>>* AddContents = nullptr;
+					if (AddField.IsValid() && AddField->TryGetArray(AddContents))
+					{
+						RemoveValueFromArray(*AddContents, CurrentValue);
+
+						// remove the "+" array if it's now empty, since it's unnecessary
+						if (AddContents->Num() == 0)
+						{
+							(*ArrayObject)->RemoveField(TEXT("+"));
+						}
+					}
+				}
+
+				// if there's a value set in the parent, then it might have been removed in this config, so we need to remove that
+				TSharedPtr<FJsonValue> ParentValue;
+				if (ParentConfig.IsValid())
+				{
+					if (ParentConfig->TryGetJsonValue(Path, ParentValue))
+					{
+						TSharedPtr<FJsonValue> RemoveField = (*ArrayObject)->TryGetField(TEXT("-"));
+						TArray<TSharedPtr<FJsonValue>>* RemoveContents = nullptr;
+						if (RemoveField.IsValid() && RemoveField->TryGetArray(RemoveContents))
+						{
+							RemoveValueFromArray(*RemoveContents, ParentValue);
+
+							// remove the "-" array if it's now empty, since it's unnecessary
+							if (RemoveContents->Num() == 0)
+							{
+								(*ArrayObject)->RemoveField(TEXT("-"));
+							}
+						}
+					}
+				}
+			}
+
+			return true;
+		}
+
+		(*CurrentObject)->RemoveField(LastPart.Name);
+
+		return true;
+	}
+
+	bool FJsonConfig::SetJsonValue(const FJsonPath& Path, TSharedPtr<FJsonValue> NewValue)
+	{
+		TSharedPtr<FJsonValue> PreviousValue;
+		TryGetJsonValue(Path, PreviousValue);
+
+		TSharedPtr<FJsonValue> MergedValue = FJsonValue::Duplicate(NewValue);
+		if (!SetJsonValueInMerged(Path, MergedValue))
+		{
+			return false;
+		}
+
+		TSharedPtr<FJsonValue> ParentValue;
+
+		// this was a valid change in the merged object 
+		// we need to either add or remove it from this object, depending on if it differs from our parent
+		bool bShouldRemove = false;
+		if (ParentConfig.IsValid())
+		{
+			if (ParentConfig->TryGetJsonValue(Path, ParentValue))
+			{
+				if (FJsonValue::CompareEqual(*ParentValue.Get(), *NewValue.Get()))
+				{
+					// same as inherited, remove it from this
+					bShouldRemove = true;
+				}
+			}
+		}
+
+		if (bShouldRemove)
+		{
+			return RemoveJsonValueFromOverride(Path, PreviousValue);
+		}
+
+		TSharedPtr<FJsonValue> OverrideValue = FJsonValue::Duplicate(NewValue);
+		if (!SetJsonValueInOverride(Path, OverrideValue, PreviousValue, ParentValue))
+		{
+			return false;
+		}
+
+		OnConfigChanged.ExecuteIfBound();
+		return true;
+	}
+
+	bool FJsonConfig::SetString(const FJsonPath& Path, const FText& Value)
+	{
+		TSharedPtr<FJsonValue> JsonValue = MakeShared<FJsonValueString>(Value.ToString());
+		return SetJsonValue(Path, JsonValue);
+	}
+
+	bool FJsonConfig::SetString(const FJsonPath& Path, FStringView Value)
+	{
+		TSharedPtr<FJsonValue> JsonValue = MakeShared<FJsonValueString>(FString(Value));
+		return SetJsonValue(Path, JsonValue);
+	}
+
+	bool FJsonConfig::SetBool(const FJsonPath& Path, bool Value)
+	{
+		TSharedPtr<FJsonValue> JsonValue = MakeShared<FJsonValueBoolean>(Value);
+		return SetJsonValue(Path, JsonValue);
+	}
+
+	bool FJsonConfig::SetJsonObject(const FJsonPath& Path, TSharedPtr<FJsonObject> Object)
+	{
+		TSharedPtr<FJsonValue> JsonValue = MakeShared<FJsonValueObject>(Object);
+		return SetJsonValue(Path, JsonValue);
+	}
+
+	bool FJsonConfig::SetJsonArray(const FJsonPath& Path, const TArray<TSharedPtr<FJsonValue>>& Array)
+	{
+		TSharedPtr<FJsonValue> JsonValue = MakeShared<FJsonValueArray>(Array);
+		return SetJsonValue(Path, JsonValue);
+	}
+
+	bool FJsonConfig::HasOverride(const FJsonPath& Path) const
+	{
+		if (!IsValid())
+		{
+			return false;
+		}
+
+		if (!Path.IsValid())
+		{
+			return false;
+		}
+
+		const TSharedPtr<FJsonObject>* CurrentObject = &OverrideObject;
+
+		for (int32 Idx = 0; Idx < Path.Length() - 1; ++Idx)
+		{
+			const TSharedPtr<FJsonValue>* NextValue = FindPartInObject(*CurrentObject, Path[Idx]);
+			if (NextValue == nullptr || !(*NextValue)->TryGetObject(CurrentObject))
+			{
+				return false;
+			}
+		}
+
+		const TSharedPtr<FJsonValue>* LastValue = FindPartInObject(*CurrentObject, Path[Path.Length() - 1]);
+		if (LastValue == nullptr)
+		{
+			return false;
+		}
+
+		return true;
+	}
+
+	bool FJsonConfig::SetRootObject(TSharedPtr<FJsonObject> Object)
+	{
+		if (!Object.IsValid())
+		{
+			return false;
+		}
+
+		TArray<FString> NewKeys;
+		Object->Values.GenerateKeyArray(NewKeys);
+
+		TArray<FString> ExistingKeys;
+		OverrideObject->Values.GenerateKeyArray(ExistingKeys);
+
+		for (const TPair<FString, TSharedPtr<FJsonValue>>& Pair : Object->Values)
+		{
+			FJsonPath Path(Pair.Key);
+			SetJsonValue(Path, Pair.Value);
+		}
+
+		// remove all keys that are not in the new object
+		for (const FString& Key : ExistingKeys)
+		{
+			if (!NewKeys.Contains(Key))
+			{
+				OverrideObject->Values.Remove(Key);
+			}
+		}
+
+		MergeThisWithParent();
+		OnConfigChanged.ExecuteIfBound();
+
+		return true;
+	}
+}
diff --git a/Engine/Source/Editor/EditorConfig/Private/Tests/EditorConfigTestHelpers.cpp b/Engine/Source/Editor/EditorConfig/Private/Tests/EditorConfigTestHelpers.cpp
new file mode 100644
--- /dev/null
+++ b/Engine/Source/Editor/EditorConfig/Private/Tests/EditorConfigTestHelpers.cpp
@@ -0,0 +1,26 @@
+// Copyright Epic Games, Inc. All Rights Reserved.
+
+#include "EditorConfigTestHelpers.h"
+#include "Serialization/JsonSerializer.h"
+
+namespace FEditorConfigTestHelpers
+{
+	bool AreJsonStringsEquivalent(const FString& Actual, const FString& Expected)
+	{
+		TSharedPtr<FJsonObject> ActualObject = MakeShared<FJsonObject>();
+		TSharedRef<FJsonStringReader> ActualReader = FJsonStringReader::Create(Actual);
+		if (!FJsonSerializer::Deserialize(ActualReader.Get(), ActualObject))
+		{
+			return false;
+		}
+
+		TSharedPtr<FJsonObject> ExpectedObject = MakeShared<FJsonObject>();
+		TSharedRef<FJsonStringReader> ExpectedReader = FJsonStringReader::Create(Expected);
+		if (!FJsonSerializer::Deserialize(ExpectedReader.Get(), ExpectedObject))
+		{
+			return false;
+		}
+
+		return FJsonValue::CompareEqual(FJsonValueObject(ActualObject), FJsonValueObject(ExpectedObject));
+	}
+}
\ No newline at end of file
diff --git a/Engine/Source/Editor/EditorConfig/Private/Tests/EditorConfigTestHelpers.h b/Engine/Source/Editor/EditorConfig/Private/Tests/EditorConfigTestHelpers.h
new file mode 100644
--- /dev/null
+++ b/Engine/Source/Editor/EditorConfig/Private/Tests/EditorConfigTestHelpers.h
@@ -0,0 +1,9 @@
+// Copyright Epic Games, Inc. All Rights Reserved.
+
+#pragma once
+
+namespace FEditorConfigTestHelpers
+{
+	// check that two serialized JSON objects are equivalent
+	static bool AreJsonStringsEquivalent(const FString& Actual, const FString& Expected);
+}
\ No newline at end of file
diff --git a/Engine/Source/Editor/EditorConfig/Private/Tests/EditorConfigTests.h b/Engine/Source/Editor/EditorConfig/Private/Tests/EditorConfigTests.h
new file mode 100644
--- /dev/null
+++ b/Engine/Source/Editor/EditorConfig/Private/Tests/EditorConfigTests.h
@@ -0,0 +1,110 @@
+#pragma once
+
+#include "CoreMinimal.h"
+#include "UObject/ObjectMacros.h"
+
+#include "EditorConfigTests.generated.h"
+
+USTRUCT()
+struct FEditorConfigTestStruct
+{
+	GENERATED_BODY()
+
+	UPROPERTY()
+	bool BoolProperty;
+
+	UPROPERTY()
+	int32 IntProperty;
+
+	UPROPERTY()
+	FString StringProperty;
+
+	UPROPERTY()
+	float FloatProperty;
+
+	UPROPERTY()
+	TArray<FString> ArrayProperty;
+};
+
+USTRUCT()
+struct FEditorConfigTestKeyStruct
+{
+	GENERATED_BODY()
+	
+	UPROPERTY()
+	FString Name;
+	
+	UPROPERTY()
+	double Number;
+
+	bool operator==(const FEditorConfigTestKeyStruct& Other) const;
+	bool operator!=(const FEditorConfigTestKeyStruct& Other) const;
+};
+
+FORCEINLINE uint32 GetTypeHash(const FEditorConfigTestKeyStruct& Struct)
+{
+	return HashCombine(GetTypeHash(Struct.Name), GetTypeHash(Struct.Number));
+}
+
+USTRUCT()
+struct FEditorConfigTestComplexArrayStruct
+{
+	GENERATED_BODY()
+	
+	UPROPERTY()
+	TArray<FEditorConfigTestKeyStruct> Array;
+};
+
+USTRUCT()
+struct FEditorConfigTestSimpleMapStruct
+{
+	GENERATED_BODY()
+
+	UPROPERTY()
+	TMap<FString, FString> Map;
+};
+
+USTRUCT()
+struct FEditorConfigTestComplexMapStruct
+{
+	GENERATED_BODY()
+	
+	UPROPERTY()
+	TMap<FEditorConfigTestKeyStruct, FEditorConfigTestKeyStruct> Map;
+};
+
+USTRUCT()
+struct FEditorConfigTestSimpleSetStruct
+{
+	GENERATED_BODY()
+	
+	UPROPERTY()
+	TSet<FName> Set;
+};
+
+USTRUCT()
+struct FEditorConfigTestComplexSetStruct
+{
+	GENERATED_BODY()
+	
+	UPROPERTY()
+	TSet<FEditorConfigTestKeyStruct> Set;
+};
+
+UCLASS()
+class UEditorConfigTestObject : public UObject
+{
+	GENERATED_BODY()
+
+	UPROPERTY()
+	UObject* Object;
+
+	UPROPERTY()
+	FSoftObjectPath SoftObjectPath;
+
+	UPROPERTY()
+	FEditorConfigTestStruct Struct;
+
+	UPROPERTY()
+	int32 Number;
+};
diff --git a/Engine/Source/Editor/EditorConfig/Private/Tests/JsonConfigTests.cpp b/Engine/Source/Editor/EditorConfig/Private/Tests/JsonConfigTests.cpp
new file mode 100644
--- /dev/null
+++ b/Engine/Source/Editor/EditorConfig/Private/Tests/JsonConfigTests.cpp
@@ -0,0 +1,1296 @@
+// Copyright Epic Games, Inc. All Rights Reserved.
+
+
+#if WITH_DEV_AUTOMATION_TESTS 
+
+#include "EditorConfigTestHelpers.h"
+#include "JsonConfig.h"
+#include "Misc/AutomationTest.h"
+
+using namespace UE;
+
+IMPLEMENT_SIMPLE_AUTOMATION_TEST(FJsonConfigTests_LoadBool, "JsonConfig.Load.Bool", EAutomationTestFlags::EditorContext | EAutomationTestFlags::EngineFilter)
+
+bool FJsonConfigTests_LoadBool::RunTest(const FString& Parameters)
+{
+	const FString Contents = 
+R"_JSON({ 
+	"Bool": true
+}
+)_JSON";
+
+	FJsonConfig Config;
+	if (!Config.LoadFromString(Contents))
+	{
+		return false;
+	}
+
+	bool Bool = false;
+	if (!Config.TryGetBool(TEXT("Bool"), Bool))
+	{
+		return false;
+	}
+
+	TestEqual(TEXT("Bool"), Bool, true);
+
+	return true;
+}
+
+IMPLEMENT_SIMPLE_AUTOMATION_TEST(FJsonConfigTests_LoadString, "JsonConfig.Load.String", EAutomationTestFlags::EditorContext | EAutomationTestFlags::EngineFilter)
+
+bool FJsonConfigTests_LoadString::RunTest(const FString& Parameters)
+{
+	const FString Contents = 
+R"_JSON({ 
+	"Foo": "Foo",
+	"Bar": "Bar"
+}
+)_JSON";
+
+	FJsonConfig Config;
+	if (!Config.LoadFromString(Contents))
+	{
+		return false;
+	}
+
+	FString Foo, Bar;
+	if (!Config.TryGetString(TEXT("Foo"), Foo) ||
+		!Config.TryGetString(TEXT("Bar"), Bar))
+	{
+		return false;
+	}
+
+	TestEqual(TEXT("Foo"), Foo, TEXT("Foo"));
+	TestEqual(TEXT("Bar"), Bar, TEXT("Bar"));
+
+	return true;
+}
+
+IMPLEMENT_SIMPLE_AUTOMATION_TEST(FJsonConfigTests_LoadNumber, "JsonConfig.Load.Number", EAutomationTestFlags::EditorContext | EAutomationTestFlags::EngineFilter)
+
+bool FJsonConfigTests_LoadNumber::RunTest(const FString& Parameters)
+{
+	const FString Contents = 
+R"_JSON({ 
+	"Positive": 63,
+	"Negative": -63,
+	"PositiveFloat": 12.5,
+	"NegativeFloat": -12.5
+}
+)_JSON";
+
+	FJsonConfig Config;
+	if (!Config.LoadFromString(Contents))
+	{
+		return false;
+	}
+
+	uint8 U8;
+	uint16 U16;
+	uint32 U32;
+	uint64 U64;
+	int8 I8;
+	int16 I16;
+	int32 I32;
+	int64 I64;
+	float Float = 0;
+	double Double = 0;
+
+	if (!Config.TryGetNumber(TEXT("Positive"), U8) ||
+		!Config.TryGetNumber(TEXT("Positive"), U16) ||
+		!Config.TryGetNumber(TEXT("Positive"), U32) ||
+		!Config.TryGetNumber(TEXT("Positive"), U64) ||
+		!Config.TryGetNumber(TEXT("Positive"), I8) ||
+		!Config.TryGetNumber(TEXT("Positive"), I16) ||
+		!Config.TryGetNumber(TEXT("Positive"), I32) ||
+		!Config.TryGetNumber(TEXT("Positive"), I64) ||
+		!Config.TryGetNumber(TEXT("Positive"), Float) ||
+		!Config.TryGetNumber(TEXT("Positive"), Double))
+	{
+		return false;
+	}
+
+	TestEqual<uint8>(TEXT("U8"), U8, 63);
+	TestEqual<uint16>(TEXT("U16"), U16, 63);
+	TestEqual<uint32>(TEXT("U32"), U32, 63);
+	TestEqual<uint64>(TEXT("U64"), U64, 63);
+	TestEqual<int8>(TEXT("I8"), I8, 63);
+	TestEqual<int16>(TEXT("I16"), I16, 63);
+	TestEqual<int32>(TEXT("I32"), I32, 63);
+	TestEqual<int64>(TEXT("I64"), I64, 63);
+	TestEqual(TEXT("Float"), Float, 63.f);
+	TestEqual(TEXT("Double"), Double, 63.0);
+
+	if (!Config.TryGetNumber(TEXT("Negative"), I8) ||
+		!Config.TryGetNumber(TEXT("Negative"), I16) ||
+		!Config.TryGetNumber(TEXT("Negative"), I32) ||
+		!Config.TryGetNumber(TEXT("Negative"), I64) ||
+		!Config.TryGetNumber(TEXT("Negative"), Float) ||
+		!Config.TryGetNumber(TEXT("Negative"), Double))
+	{
+		return false;
+	}
+
+	TestEqual<int8>(TEXT("I8"), I8, -63);
+	TestEqual<int16>(TEXT("I16"), I16, -63);
+	TestEqual<int32>(TEXT("I32"), I32, -63);
+	TestEqual<int64>(TEXT("I64"), I64, -63);
+	TestEqual(TEXT("Float"), Float, -63.f);
+	TestEqual(TEXT("Double"), Double, -63.0);
+
+	if (!Config.TryGetNumber(TEXT("PositiveFloat"), Float) ||
+		!Config.TryGetNumber(TEXT("PositiveFloat"), Double))
+	{
+		return false;
+	}
+
+	TestEqual(TEXT("Float"), Float, 12.5f);
+	TestEqual(TEXT("Double"), Double, 12.5);
+
+	if (!Config.TryGetNumber(TEXT("NegativeFloat"), Float) ||
+		!Config.TryGetNumber(TEXT("NegativeFloat"), Double))
+	{
+		return false;
+	}
+
+	TestEqual(TEXT("Float"), Float, -12.5f);
+	TestEqual(TEXT("Double"), Double, -12.5);
+
+	return true;
+}
+
+IMPLEMENT_SIMPLE_AUTOMATION_TEST(FJsonConfigTests_LoadArrayBool, "JsonConfig.Load.Array.Bool", EAutomationTestFlags::EditorContext | EAutomationTestFlags::EngineFilter)
+
+bool FJsonConfigTests_LoadArrayBool::RunTest(const FString& Parameters)
+{
+	const FString Contents = 
+R"_JSON({ 
+	"Array": [true, false, true]
+}
+)_JSON";
+
+	FJsonConfig Config;
+	if (!Config.LoadFromString(Contents))
+	{
+		return false;
+	}
+
+	TArray<bool> Array;
+	if (!Config.TryGetArray(TEXT("Array"), Array))
+	{
+		return false;
+	}
+
+	TestEqual(TEXT("Size"), Array.Num(), 3);
+	TestEqual(TEXT("Contents"), Array, TArray<bool> { true, false, true });
+
+	bool At1;
+	if (!Config.TryGetBool(TEXT("Array[1]"), At1))
+	{
+		return false;
+	} 
+
+	TestFalse(TEXT("Indexed"), At1);
+
+	return true;
+}
+
+IMPLEMENT_SIMPLE_AUTOMATION_TEST(FJsonConfigTests_LoadArraySetBool, "JsonConfig.Load.Array.Set.Bool", EAutomationTestFlags::EditorContext | EAutomationTestFlags::EngineFilter)
+
+bool FJsonConfigTests_LoadArraySetBool::RunTest(const FString& Parameters)
+{
+	const FString Contents = 
+R"_JSON({ 
+	"Array": 
+	{
+		"=": [true, false, true]
+	}
+}
+)_JSON";
+
+	FJsonConfig Config;
+	if (!Config.LoadFromString(Contents))
+	{
+		return false;
+	}
+
+	TArray<bool> Array;
+	if (!Config.TryGetArray(TEXT("Array"), Array))
+	{
+		return false;
+	}
+
+	TestEqual(TEXT("Size"), Array.Num(), 3);
+	TestEqual(TEXT("Contents"), Array, TArray<bool> { true, false, true });
+
+	return true;
+}
+
+IMPLEMENT_SIMPLE_AUTOMATION_TEST(FJsonConfigTests_LoadArrayAddBool, "JsonConfig.Load.Array.Add.Bool", EAutomationTestFlags::EditorContext | EAutomationTestFlags::EngineFilter)
+
+bool FJsonConfigTests_LoadArrayAddBool::RunTest(const FString& Parameters)
+{
+	const FString Contents = 
+R"_JSON({ 
+	"Array": 
+	{
+		"+": [true, false, true]
+	}
+}
+)_JSON";
+
+	FJsonConfig Config;
+	if (!Config.LoadFromString(Contents))
+	{
+		return false;
+	}
+
+	TArray<bool> Array;
+	if (!Config.TryGetArray(TEXT("Array"), Array))
+	{
+		return false;
+	}
+
+	TestEqual(TEXT("Size"), Array.Num(), 3);
+	TestEqual(TEXT("Contents"), Array, TArray<bool> { true, false, true });
+	return true;
+}
+
+IMPLEMENT_SIMPLE_AUTOMATION_TEST(FJsonConfigTests_LoadNestedObjects, "JsonConfig.Load.Objects.Nested", EAutomationTestFlags::EditorContext | EAutomationTestFlags::EngineFilter)
+
+bool FJsonConfigTests_LoadNestedObjects::RunTest(const FString& Parameters)
+{
+	const FString Contents = 
+R"_JSON({ 
+	"Top": 
+	{
+		"Nested":
+		{
+			"Property": "Foo"
+		}
+	}
+}
+)_JSON";
+
+	FJsonConfig Config;
+	if (!Config.LoadFromString(Contents))
+	{
+		return false;
+	}
+
+	TSharedPtr<FJsonObject> Top;
+	if (!Config.TryGetJsonObject(TEXT("Top"), Top))
+	{
+		return false;
+	}
+
+	TestEqual("Top", Top->Values.Num(), 1);
+
+	TSharedPtr<FJsonObject> Nested;
+	if (!Config.TryGetJsonObject(TEXT("Top.Nested"), Nested))
+	{
+		return false;
+	}
+
+	TestEqual(TEXT("Nested"), Top->Values[TEXT("Nested")]->AsObject(), Nested);
+
+	FString Property;
+	if (!Config.TryGetString(TEXT("Top.Nested.Property"), Property))
+	{
+		return false;
+	}
+
+	TestEqual(TEXT("Property"), Property, TEXT("Foo"));
+
+	return true;
+}
+
+IMPLEMENT_SIMPLE_AUTOMATION_TEST(FJsonConfigTests_LoadMapNumericKey, "JsonConfig.Load.Map.NumericKey", EAutomationTestFlags::EditorContext | EAutomationTestFlags::EngineFilter)
+
+bool FJsonConfigTests_LoadMapNumericKey::RunTest(const FString& Parameters)
+{
+	const FString Contents = 
+R"_JSON({ 
+	"Map":
+	[
+		{
+			"$key": 5,
+			"$value": "Foo"
+		},
+		{
+			"$key": 10,
+			"$value": "Bar"
+		}
+	]
+}
+)_JSON";
+
+	FJsonConfig Config;
+	if (!Config.LoadFromString(Contents))
+	{
+		return false;
+	}
+
+	TArray<FJsonValuePair> Map;
+	if (!Config.TryGetMap(TEXT("Map"), Map))
+	{
+		return false;
+	}
+
+	return true;
+}
+
+IMPLEMENT_SIMPLE_AUTOMATION_TEST(FJsonConfigTests_LoadHierarchicalBool, "JsonConfig.Load.Hierarchical.Bool", EAutomationTestFlags::EditorContext | EAutomationTestFlags::EngineFilter)
+
+bool FJsonConfigTests_LoadHierarchicalBool::RunTest(const FString& Parameters)
+{
+	const FString ParentContents = 
+R"_JSON({ 
+	"Bool": true
+}
+)_JSON";
+
+	TSharedRef<FJsonConfig> ParentConfig = MakeShared<FJsonConfig>();
+	if (!ParentConfig->LoadFromString(ParentContents))
+	{
+		return false;
+	}
+
+	bool Value;
+	if (!ParentConfig->TryGetBool(TEXT("Bool"), Value))
+	{
+		return false;
+	}
+
+	TestEqual("Parent Value", Value, true);
+
+	const FString ChildContents = 
+R"_JSON({ 
+	"Bool": false
+}
+)_JSON";
+
+	FJsonConfig ChildConfig;
+	ChildConfig.SetParent(ParentConfig);
+
+	if (!ChildConfig.TryGetBool(TEXT("Bool"), Value))
+	{
+		return false;
+	}
+
+	TestEqual("Inherited Value", Value, true);
+
+	if (!ChildConfig.LoadFromString(ChildContents))
+	{
+		return false;
+	}
+
+	if (!ChildConfig.TryGetBool(TEXT("Bool"), Value))
+	{
+		return false;
+	}
+
+	TestEqual("Child Value", Value, false);
+
+	// reparent and check
+	ChildConfig.SetParent(ParentConfig);
+
+	if (!ChildConfig.TryGetBool(TEXT("Bool"), Value))
+	{
+		return false;
+	}
+
+	TestEqual("Reparented Value", Value, false);
+
+	return true;
+}
+
+IMPLEMENT_SIMPLE_AUTOMATION_TEST(FJsonConfigTests_LoadHierarchicalString, "JsonConfig.Load.Hierarchical.String", EAutomationTestFlags::EditorContext | EAutomationTestFlags::EngineFilter)
+
+bool FJsonConfigTests_LoadHierarchicalString::RunTest(const FString& Parameters)
+{
+	const FString ParentContents = 
+R"_JSON({ 
+	"String": "Foo"
+}
+)_JSON";
+
+	TSharedRef<FJsonConfig> ParentConfig = MakeShared<FJsonConfig>();
+	if (!ParentConfig->LoadFromString(ParentContents))
+	{
+		return false;
+	}
+
+	FString String;
+	if (!ParentConfig->TryGetString(TEXT("String"), String))
+	{
+		return false;
+	}
+
+	TestEqual("Parent Value", String, TEXT("Foo"));
+
+	const FString ChildContents = 
+R"_JSON({ 
+	"String": "Bar"
+}
+)_JSON";
+
+	FJsonConfig ChildConfig;
+	ChildConfig.SetParent(ParentConfig);
+
+	if (!ChildConfig.TryGetString(TEXT("String"), String))
+	{
+		return false;
+	}
+
+	TestEqual("Inherited Value", String, TEXT("Foo"));
+
+	if (!ChildConfig.LoadFromString(ChildContents))
+	{
+		return false;
+	}
+
+	if (!ChildConfig.TryGetString(TEXT("String"), String))
+	{
+		return false;
+	}
+
+	TestEqual("Child Value", String, TEXT("Bar"));
+
+	return true;
+}
+
+IMPLEMENT_SIMPLE_AUTOMATION_TEST(FJsonConfigTests_LoadHierarchicalObject, "JsonConfig.Load.Hierarchical.Object", EAutomationTestFlags::EditorContext | EAutomationTestFlags::EngineFilter)
+
+bool FJsonConfigTests_LoadHierarchicalObject::RunTest(const FString& Parameters)
+{
+	const FString ParentContents = 
+R"_JSON({ 
+	"Object":
+	{
+		"Number": 1,
+		"OldString": "Foo"
+	}
+}
+)_JSON";
+
+	TSharedRef<FJsonConfig> ParentConfig = MakeShared<FJsonConfig>();
+	if (!ParentConfig->LoadFromString(ParentContents))
+	{
+		return false;
+	}
+
+	int32 Int;
+	if (!ParentConfig->TryGetNumber(TEXT("Object.Number"), Int))
+	{
+		return false;
+	}
+
+	TestEqual("Parent Number", Int, 1);
+
+	FString OldString;
+	if (!ParentConfig->TryGetString(TEXT("Object.OldString"), OldString))
+	{
+		return false;
+	}
+
+	// check for non-existence of the new string
+	FString NewString;
+	if (ParentConfig->TryGetString(TEXT("Object.NewString"), NewString))
+	{
+		return false;
+	}
+
+	TestEqual("Parent Old String", OldString, TEXT("Foo"));
+
+	const FString ChildContents = 
+R"_JSON({ 
+	"Object":
+	{
+		"Number": 5,
+		"NewString": "Bar"
+	}
+}
+)_JSON";
+
+	FJsonConfig ChildConfig;
+	ChildConfig.SetParent(ParentConfig);
+
+	if (!ChildConfig.TryGetNumber(TEXT("Object.Number"), Int))
+	{
+		return false;
+	}
+
+	TestEqual("Inherited Int", Int, 1);
+
+	if (!ChildConfig.LoadFromString(ChildContents))
+	{
+		return false;
+	}
+
+	if (!ChildConfig.TryGetNumber(TEXT("Object.Number"), Int))
+	{
+		return false;
+	}
+
+	TestEqual("Child Int", Int, 5);
+
+	if (!ChildConfig.TryGetString(TEXT("Object.OldString"), OldString))
+	{
+		return false;
+	}
+
+	TestEqual("Child Old String", OldString, TEXT("Foo"));
+
+	if (!ChildConfig.TryGetString(TEXT("Object.NewString"), NewString))
+	{
+		return false;
+	}
+
+	TestEqual("Child New String", NewString, TEXT("Bar"));
+
+	return true;
+}
+
+
+IMPLEMENT_SIMPLE_AUTOMATION_TEST(FJsonConfigTests_LoadHierarchicalMapStringKeyRemove, "JsonConfig.Load.Hierarchical.Map.StringKey.Remove", EAutomationTestFlags::EditorContext | EAutomationTestFlags::EngineFilter)
+
+bool FJsonConfigTests_LoadHierarchicalMapStringKeyRemove::RunTest(const FString& Parameters)
+{
+	const FString ParentContents = 
+R"_JSON({ 
+	"Map":
+	[
+		{
+			"$key": "Foo",
+			"$value": true
+		},
+		{
+			"$key": "Bar",
+			"$value": true
+		}
+	]
+}
+)_JSON";
+
+	TSharedRef<FJsonConfig> ParentConfig = MakeShared<FJsonConfig>();
+	if (!ParentConfig->LoadFromString(ParentContents))
+	{
+		return false;
+	}
+
+	TArray<FJsonValuePair> Map;
+	if (!ParentConfig->TryGetMap(TEXT("Map"), Map))
+	{
+		return false;
+	}
+
+	TestEqual(TEXT("Parent Map Size"), Map.Num(), 2);
+	
+	FString Key = Map[0].Key->AsString();
+	TestEqual(TEXT("Parent Map Key"), Key, TEXT("Foo"));
+
+	bool Value = Map[0].Value->AsBool();
+	TestEqual(TEXT("Parent Map Value"), Value, true);
+
+	Key = Map[1].Key->AsString();
+	TestEqual(TEXT("Parent Map Key"), Key, TEXT("Bar"));
+
+	Value = Map[1].Value->AsBool();
+	TestEqual(TEXT("Parent Map Value"), Value, true);
+
+	const FString ChildContents = 
+R"_JSON({ 
+	"Map":
+	{
+		"-":
+		[
+			"Foo"
+		]
+	}
+}
+)_JSON";
+
+	FJsonConfig ChildConfig;
+	if (!ChildConfig.LoadFromString(ChildContents))
+	{
+		return false;
+	}
+
+	if (!ChildConfig.TryGetMap(TEXT("Map"), Map))
+	{
+		return false;
+	}
+
+	TestEqual(TEXT("Child Map Size"), Map.Num(), 0);
+
+	ChildConfig.SetParent(ParentConfig);
+
+	if (!ChildConfig.TryGetMap(TEXT("Map"), Map))
+	{
+		return false;
+	}
+
+	TestEqual(TEXT("Inherited Map Size"), Map.Num(), 1);
+
+	TestEqual(TEXT("Inherited Map Key"), Map[0].Key->AsString(), TEXT("Bar"));
+	TestEqual(TEXT("Inherited Map Value"), Map[0].Value->AsBool(), true);
+
+	return true;
+}
+
+IMPLEMENT_SIMPLE_AUTOMATION_TEST(FJsonConfigTests_LoadHierarchicalMapNumericKeyAdd, "JsonConfig.Load.Hierarchical.Map.NumericKey.Add", EAutomationTestFlags::EditorContext | EAutomationTestFlags::EngineFilter)
+
+bool FJsonConfigTests_LoadHierarchicalMapNumericKeyAdd::RunTest(const FString& Parameters)
+{
+	const FString ParentContents = 
+R"_JSON({ 
+	"Map":
+	[
+		{
+			"$key": 5,
+			"$value": "Foo"
+		},
+		{
+			"$key": 20,
+			"$value": "Removed"
+		}
+	]
+}
+)_JSON";
+
+	TSharedRef<FJsonConfig> ParentConfig = MakeShared<FJsonConfig>();
+	if (!ParentConfig->LoadFromString(ParentContents))
+	{
+		return false;
+	}
+
+	TArray<FJsonValuePair> Map;
+	if (!ParentConfig->TryGetMap(TEXT("Map"), Map))
+	{
+		return false;
+	}
+
+	TestEqual("Parent Map Size", Map.Num(), 2);
+	
+	double Key = Map[0].Key->AsNumber();
+	TestEqual("Parent Map Key", Key, 5.0);
+
+	FString Value = Map[0].Value->AsString();
+	TestEqual("Parent Map Value", Value, TEXT("Foo"));
+
+	Key = Map[1].Key->AsNumber();
+	TestEqual("Parent Map Key", Key, 20.0);
+
+	Value = Map[1].Value->AsString();
+	TestEqual("Parent Map Value", Value, TEXT("Removed"));
+
+	const FString ChildContents = 
+R"_JSON({ 
+	"Map":
+	{
+		"+":
+		[
+			{
+				"$key": 10,
+				"$value": "Bar"
+			}
+		],
+		"-":
+		[
+			{
+				"$key": 20
+			}
+		]
+	}
+}
+)_JSON";
+
+	FJsonConfig ChildConfig;
+	if (!ChildConfig.LoadFromString(ChildContents))
+	{
+		return false;
+	}
+
+	if (!ChildConfig.TryGetMap(TEXT("Map"), Map))
+	{
+		return false;
+	}
+
+	TestEqual("Child Map Size", Map.Num(), 1);
+
+	ChildConfig.SetParent(ParentConfig);
+
+	if (!ChildConfig.TryGetMap(TEXT("Map"), Map))
+	{
+		return false;
+	}
+
+	TestEqual("Inherited Map Size", Map.Num(), 2);
+
+	for (const FJsonValuePair& Pair : Map)
+	{
+		Key = Pair.Key->AsNumber();
+		FString Expected;
+		if (Key == 5.0)
+		{
+			Expected = TEXT("Foo");
+		}
+		else if (Key == 10.0)
+		{
+			Expected = TEXT("Bar");
+		}
+		else
+		{
+			checkf(false, TEXT("Invalid key found in JSON map."));
+		}
+
+		Value = Pair.Value->AsString();
+		TestEqual("Inherited Map Value", Value, Expected);
+	}
+
+	return true;
+}
+
+IMPLEMENT_SIMPLE_AUTOMATION_TEST(FJsonConfigTests_LoadHierarchicalMapObjectKeyAdd, "JsonConfig.Load.Hierarchical.Map.ObjectKey.Add", EAutomationTestFlags::EditorContext | EAutomationTestFlags::EngineFilter)
+
+bool FJsonConfigTests_LoadHierarchicalMapObjectKeyAdd::RunTest(const FString& Parameters)
+{
+	const FString ParentContents = 
+R"_JSON({ 
+	"Map":
+	[
+		{
+			"$key": 
+			{
+				"Key": "Foo"
+			},
+			"$value": "Foo"
+		}
+	]
+}
+)_JSON";
+
+	TSharedRef<FJsonConfig> ParentConfig = MakeShared<FJsonConfig>();
+	if (!ParentConfig->LoadFromString(ParentContents))
+	{
+		return false;
+	}
+
+	TArray<FJsonValuePair> Map;
+	if (!ParentConfig->TryGetMap(TEXT("Map"), Map))
+	{
+		return false;
+	}
+
+	TestEqual("Parent Map Size", Map.Num(), 1);
+
+	TSharedPtr<FJsonObject> KeyObject = Map[0].Key->AsObject();
+	FString KeyString = KeyObject->GetStringField(TEXT("Key"));
+	TestEqual("Parent Map Key", KeyString, TEXT("Foo"));
+
+	TestEqual("Parent Map Value", Map[0].Value->AsString(), TEXT("Foo"));
+
+	const FString ChildContents = 
+R"_JSON({ 
+	"Map":
+	{
+		"+":
+		[
+			{
+				"$key": 
+				{
+					"Key": "Bar"
+				},
+				"$value": "Bar"
+			}
+		]
+	}
+}
+)_JSON";
+
+	FJsonConfig ChildConfig;
+	if (!ChildConfig.LoadFromString(ChildContents))
+	{
+		return false;
+	}
+
+	if (!ChildConfig.TryGetMap(TEXT("Map"), Map))
+	{
+		return false;
+	}
+
+	TestEqual("Child Map Size", Map.Num(), 1);
+
+	ChildConfig.SetParent(ParentConfig);
+
+	if (!ChildConfig.TryGetMap(TEXT("Map"), Map))
+	{
+		return false;
+	}
+
+	TestEqual("Inherited Map Size", Map.Num(), 2);
+
+	KeyObject = Map[0].Key->AsObject();
+	KeyString = KeyObject->GetStringField(TEXT("Key"));
+	TestEqual("Inherited Map Key", KeyString, TEXT("Foo"));
+	TestEqual("Inherited Map Value", Map[0].Value->AsString(), TEXT("Foo"));
+
+	KeyObject = Map[1].Key->AsObject();
+	KeyString = KeyObject->GetStringField(TEXT("Key"));
+	TestEqual("Inherited Map Key", KeyString, TEXT("Bar"));
+	TestEqual("Inherited Map Value", Map[1].Value->AsString(), TEXT("Bar"));
+
+	// try to override a key with an Add
+	const FString OverrideContents = 
+R"_JSON({ 
+	"Map":
+	{
+		"+":
+		[
+			{
+				"$key": 
+				{
+					"Key": "Foo"
+				},
+				"$value": "Overridden"
+			}
+		]
+	}
+}
+)_JSON";
+
+	if (!ChildConfig.LoadFromString(OverrideContents))
+	{
+		return false;
+	}
+
+	if (!ChildConfig.TryGetMap(TEXT("Map"), Map))
+	{
+		return false;
+	}
+
+	TestEqual("Child Map Size", Map.Num(), 1);
+
+	ChildConfig.SetParent(ParentConfig);
+
+	if (!ChildConfig.TryGetMap(TEXT("Map"), Map))
+	{
+		return false;
+	}
+
+	TestEqual("Inherited Map Size", Map.Num(), 1);
+
+	KeyObject = Map[0].Key->AsObject();
+	KeyString = KeyObject->GetStringField(TEXT("Key"));
+	TestEqual("Inherited Map Key", KeyString, TEXT("Foo"));
+	TestEqual("Inherited Map Value", Map[0].Value->AsString(), TEXT("Overridden"));
+
+	return true;
+}
+
+
+IMPLEMENT_SIMPLE_AUTOMATION_TEST(FJsonConfigTests_LoadHierarchicalMapObjectKeyRemove, "JsonConfig.Load.Hierarchical.Map.ObjectKey.Remove", EAutomationTestFlags::EditorContext | EAutomationTestFlags::EngineFilter)
+
+bool FJsonConfigTests_LoadHierarchicalMapObjectKeyRemove::RunTest(const FString& Parameters)
+{
+	const FString ParentContents = 
+R"_JSON({ 
+	"Map":
+	[
+		{
+			"$key": 
+			{
+				"Key": "Foo"
+			},
+			"$value": "Foo"
+		},
+		{
+			"$key": 
+			{
+				"Key": "Bar"
+			},
+			"$value": "Bar"
+		}
+	]
+}
+)_JSON";
+
+	TSharedRef<FJsonConfig> ParentConfig = MakeShared<FJsonConfig>();
+	if (!ParentConfig->LoadFromString(ParentContents))
+	{
+		return false;
+	}
+
+	TArray<FJsonValuePair> Map;
+	if (!ParentConfig->TryGetMap(TEXT("Map"), Map))
+	{
+		return false;
+	}
+
+	TestEqual("Parent Map Size", Map.Num(), 2);
+
+	const FString ChildContents = 
+R"_JSON({ 
+	"Map":
+	{
+		"-":
+		[
+			{
+				"$key": 
+				{
+					"Key": "Bar"
+				}
+			}
+		]
+	}
+}
+)_JSON";
+
+	TSharedRef<FJsonConfig> ChildConfig = MakeShared<FJsonConfig>();
+	if (!ChildConfig->LoadFromString(ChildContents))
+	{
+		return false;
+	}
+
+	ChildConfig->SetParent(ParentConfig);
+	
+	if (!ChildConfig->TryGetMap(TEXT("Map"), Map))
+	{
+		return false;
+	}
+
+	TestEqual("Child Map Size", Map.Num(), 1);
+
+	TSharedPtr<FJsonObject> KeyObject = Map[0].Key->AsObject();
+	FString KeyString = KeyObject->GetStringField(TEXT("Key"));
+	TestEqual("Child Map Key", KeyString, TEXT("Foo"));
+
+	return true;
+}
+
+IMPLEMENT_SIMPLE_AUTOMATION_TEST(FJsonConfigTests_SetNumber, "JsonConfig.Set.Number.Simple", EAutomationTestFlags::EditorContext | EAutomationTestFlags::EngineFilter)
+
+template <typename T>
+void SetAndCheckNumber(FJsonConfig& Config, T Value)
+{
+	Config.SetNumber(TEXT("Foo"), Value);
+
+	T Actual;
+	ensureAlways(Config.TryGetNumber(TEXT("Foo"), Actual));
+	ensureAlways(Value == Actual);
+}
+
+bool FJsonConfigTests_SetNumber::RunTest(const FString& Parameters)
+{
+	FJsonConfig Config;
+
+	// these should all compile and work
+	SetAndCheckNumber(Config, (uint8) 1);
+	SetAndCheckNumber(Config, (uint16) 2);
+	SetAndCheckNumber(Config, (uint32) 3);
+	SetAndCheckNumber(Config, (uint64) 4);
+	SetAndCheckNumber(Config, (int8) 5);
+	SetAndCheckNumber(Config, (int16) 6);
+	SetAndCheckNumber(Config, (int32) 7);
+	SetAndCheckNumber(Config, (int64) 8);
+	SetAndCheckNumber(Config, (float) 9);
+	SetAndCheckNumber(Config, (double) 10);
+
+	// this shouldn't compile
+	//Config.SetNumber(TEXT("Foo"), TEXT("Bar"));
+
+	return true;
+}
+
+IMPLEMENT_SIMPLE_AUTOMATION_TEST(FJsonConfigTests_SetNumberInArray, "JsonConfig.Set.Number.InArray", EAutomationTestFlags::EditorContext | EAutomationTestFlags::EngineFilter)
+
+bool FJsonConfigTests_SetNumberInArray::RunTest(const FString& Parameters)
+{
+	const FString ParentContents =
+R"_JSON({
+	"Array": [ 1, 2, 3 ]
+})_JSON";
+
+	TSharedRef<FJsonConfig> ParentConfig = MakeShared<FJsonConfig>();
+	if (!ParentConfig->LoadFromString(ParentContents))
+	{
+		return false;
+	}
+
+	FJsonConfig ChildConfig;
+	ChildConfig.SetParent(ParentConfig);
+
+	TArray<int32> Array;
+	if (!ChildConfig.TryGetArray(TEXT("Array"), Array))
+	{
+		return false;
+	}
+
+	TestEqual(TEXT("Before Set"), Array[0], 1);
+	TestEqual(TEXT("Before Set"), Array[1], 2);
+	TestEqual(TEXT("Before Set"), Array[2], 3);
+
+	ChildConfig.SetNumber(TEXT("Array[0]"), 4);
+
+	if (!ChildConfig.TryGetArray(TEXT("Array"), Array))
+	{
+		return false;
+	}
+
+	// this is somewhat misleading, since loading the override on top of the existing could reorder this
+	TestEqual(TEXT("After Set"), Array[0], 4);
+	TestEqual(TEXT("After Set"), Array[1], 2);
+	TestEqual(TEXT("After Set"), Array[2], 3);
+
+	FString OverrideString;
+	ChildConfig.SaveToString(OverrideString);
+
+	const FString ExpectedOverrideString1 = 
+R"_JSON({
+	"Array":
+	{
+		"+": [ 4 ],
+		"-": [ 1 ]
+	}
+})_JSON";
+
+	if (!TestTrue(TEXT("Serialized"), FEditorConfigTestHelpers::AreJsonStringsEquivalent(OverrideString, ExpectedOverrideString1)))
+	{
+		return false;
+	}
+
+	// set a different value
+
+	ChildConfig.SetNumber(TEXT("Array[1]"), 5);
+
+	if (!ChildConfig.TryGetArray(TEXT("Array"), Array))
+	{
+		return false;
+	}
+
+	TestEqual(TEXT("After Set"), Array[0], 4);
+	TestEqual(TEXT("After Set"), Array[1], 5);
+	TestEqual(TEXT("After Set"), Array[2], 3);
+
+	ChildConfig.SaveToString(OverrideString);
+
+	const FString ExpectedOverrideString2 = 
+R"_JSON({
+	"Array":
+	{
+		"+": [ 4, 5 ],
+		"-": [ 1, 2 ]
+	}
+})_JSON";
+
+	if (!TestTrue(TEXT("Serialized"), FEditorConfigTestHelpers::AreJsonStringsEquivalent(OverrideString, ExpectedOverrideString2)))
+	{
+		return false;
+	}
+
+	ChildConfig.SetNumber(TEXT("Array[0]"), 1);
+
+	if (!ChildConfig.TryGetArray(TEXT("Array"), Array))
+	{
+		return false;
+	}
+
+	TestEqual(TEXT("After Set"), Array[0], 1);
+	TestEqual(TEXT("After Set"), Array[1], 5);
+	TestEqual(TEXT("After Set"), Array[2], 3);
+
+	ChildConfig.SaveToString(OverrideString);
+
+	const FString ExpectedOverrideString3 = 
+R"_JSON({
+	"Array":
+	{
+		"+": [ 5 ],
+		"-": [ 2 ]
+	}
+})_JSON";
+
+	if (!TestTrue(TEXT("Serialized"), FEditorConfigTestHelpers::AreJsonStringsEquivalent(OverrideString, ExpectedOverrideString3)))
+	{
+		return false;
+	}
+
+	return true;
+}
+
+IMPLEMENT_SIMPLE_AUTOMATION_TEST(FJsonConfigTests_SetObjectSimple, "JsonConfig.Set.Object.Simple", EAutomationTestFlags::EditorContext | EAutomationTestFlags::EngineFilter)
+
+bool FJsonConfigTests_SetObjectSimple::RunTest(const FString& Parameters)
+{
+	const FString ParentContents =
+R"_JSON({
+	"Object":
+	{
+		"Foo": 1,
+		"Bar": false
+	}
+})_JSON";
+
+	TSharedRef<FJsonConfig> ParentConfig = MakeShared<FJsonConfig>();
+	if (!ParentConfig->LoadFromString(ParentContents))
+	{
+		return false;
+	}
+
+	FJsonConfig ChildConfig;
+	ChildConfig.SetParent(ParentConfig);
+
+	TSharedPtr<FJsonObject> Object;
+	if (!ChildConfig.TryGetJsonObject(TEXT("Object"), Object))
+	{
+		return false;
+	}
+
+	TSharedPtr<FJsonValue> FooValue, BarValue;
+
+	FooValue = Object->Values["Foo"];
+	TestEqual(TEXT("Before Set"), FooValue->AsNumber(), 1.0);
+
+	BarValue = Object->Values["Bar"];
+	TestEqual(TEXT("Before Set"), BarValue->AsBool(), false);
+
+	TSharedPtr<FJsonObject> NewObject = MakeShared<FJsonObject>();
+	NewObject->SetNumberField(TEXT("Foo"), 2);
+
+	ChildConfig.SetJsonObject(TEXT("Object"), NewObject);
+
+	if (!ChildConfig.TryGetJsonObject(TEXT("Object"), Object))
+	{
+		return false;
+	}
+
+	FooValue = Object->Values["Foo"];
+	TestEqual(TEXT("After Set"), FooValue->AsNumber(), 2.0);
+
+	BarValue = Object->Values["Bar"];
+	TestEqual(TEXT("After Set"), BarValue->AsBool(), false);
+
+	const FString ExpectedOverrideString1 = 
+R"_JSON({
+	"Object":
+	{
+		"Foo": 2.0
+	}
+})_JSON";
+
+	FString OverrideString;
+	ChildConfig.SaveToString(OverrideString);
+	TestTrue(TEXT("Serialized"), FEditorConfigTestHelpers::AreJsonStringsEquivalent(OverrideString, ExpectedOverrideString1));
+
+	NewObject->SetNumberField(TEXT("Foo"), 1.0);
+	NewObject->SetBoolField(TEXT("Bar"), true);
+
+	ChildConfig.SetJsonObject(TEXT("Object"), NewObject);
+
+	FooValue = Object->Values["Foo"];
+	TestEqual(TEXT("After Set"), FooValue->AsNumber(), 1.0);
+
+	BarValue = Object->Values["Bar"];
+	TestEqual(TEXT("After Set"), BarValue->AsBool(), true);
+
+	const FString ExpectedOverrideString2 = 
+R"_JSON({
+	"Object":
+	{
+		"Bar": true
+	}
+})_JSON";
+
+	ChildConfig.SaveToString(OverrideString);
+	TestTrue(TEXT("Serialized"), FEditorConfigTestHelpers::AreJsonStringsEquivalent(OverrideString, ExpectedOverrideString2));
+
+	return true;
+}
+
+
+/*IMPLEMENT_SIMPLE_AUTOMATION_TEST(FJsonConfigTests_SetStringInObjectInArray, "JsonConfig.Set.Object.InArray", EAutomationTestFlags::EditorContext | EAutomationTestFlags::EngineFilter)
+
+bool FJsonConfigTests_SetObjectInArray::RunTest(const FString& Parameters)
+{
+	const FString ParentContents =
+R"_JSON({
+	"Array": 
+	[
+		{
+			"String": "Foo",
+			"Object":
+			{
+				"Bool": false,
+				"Number": 1
+			}
+		}
+	]
+})_JSON";
+
+	TSharedRef<FJsonConfig> ParentConfig = MakeShared<FJsonConfig>();
+	if (!ParentConfig->LoadFromString(ParentContents))
+	{
+		return false;
+	}
+
+	FJsonConfig ChildConfig;
+	ChildConfig.SetParent(ParentConfig);
+
+	TSharedPtr<FJsonObject> Object;
+	if (!ChildConfig.TryGetJsonObject(TEXT("Array[0]"), Object))
+	{
+		return false;
+	}
+
+	TSharedPtr<FJsonValue> StringValue, ObjectValue;
+
+	StringValue = Object->Values["String"];
+	TestEqual(TEXT("Before Set"), StringValue->AsString(), TEXT("Foo"));
+
+	ObjectValue = Object->Values["Object"];
+
+	bool OutBool;
+	if (!ObjectValue->AsObject()->TryGetBoolField(TEXT("Bool"), OutBool))
+	{
+		return false;
+	}
+
+	TestEqual(TEXT("Before Set"), OutBool, false);
+
+	int32 OutNumber;
+	if (!ObjectValue->AsObject()->TryGetNumberField(TEXT("Number"), OutNumber))
+	{
+		return false;
+	}
+
+	TestEqual(TEXT("Before Set"), OutNumber, 1);
+
+	//ChildConfig.SetString(TEXT("Array[0].String"), TEXT("Bar"));
+
+	if (!ChildConfig.TryGetJsonObject(TEXT("Array[0]"), Object))
+	{
+		return false;
+	}
+
+	StringValue = Object->Values["String"];
+	TestEqual(TEXT("After String Set"), StringValue->AsString(), TEXT("Bar"));
+
+	const FString ExpectedOverrideString1 = 
+R"_JSON({
+	"Array":
+	[
+		{
+			
+		}
+	]
+})_JSON";
+
+	FString OverrideString;
+	ChildConfig.SaveToString(OverrideString);
+	TestTrue(TEXT("Serialized 1"), FEditorConfigTestHelpers::AreJsonStringsEquivalent(OverrideString, ExpectedOverrideString1));
+
+	return true;
+}*/
+
+#endif
\ No newline at end of file
diff --git a/Engine/Source/Editor/EditorConfig/Public/EditorConfig.h b/Engine/Source/Editor/EditorConfig/Public/EditorConfig.h
new file mode 100644
--- /dev/null
+++ b/Engine/Source/Editor/EditorConfig/Public/EditorConfig.h
@@ -0,0 +1,217 @@
+// Copyright Epic Games, Inc. All Rights Reserved.
+
+#pragma once
+
+#include "JsonConfig.h"
+#include "Concepts/EqualityComparable.h"
+#include "Dom/JsonObject.h"
+#include "Templates/Models.h"
+
+DEFINE_LOG_CATEGORY_STATIC(LogEditorConfig, Log, All);
+
+class EDITORCONFIG_API FEditorConfig
+{
+public:
+	FEditorConfig();
+
+	void SetParent(TSharedPtr<FEditorConfig> InConfig);
+
+	bool LoadFromString(FStringView Content);
+	bool SaveToString(FString& OutResult) const;
+
+	bool IsValid() const { return JsonConfig.IsValid() && JsonConfig->IsValid(); }
+
+	TSharedPtr<FEditorConfig> GetParentConfig() const;
+
+	// UStruct & UObject
+	template <typename T>
+	bool TryGetStruct(FStringView Key, T& OutValue) const;
+	template <typename T>
+	bool TryGetUObject(FStringView Key, T& OutValue) const;
+
+	template <typename T>
+	bool TryGetRootStruct(T& OutValue) const;
+	template <typename T>
+	bool TryGetRootUObject(T& OutValue) const;
+
+	template <typename T>
+	void SetStruct(FStringView Key, const T& InValue);
+	template <typename T>
+	void SetUObject(FStringView Key, const T& InValue);
+
+	template <typename T>
+	void SetRootStruct(const T& InValue);
+	template <typename T>
+	void SetRootUObject(const T& InValue);
+
+	bool HasOverride(FStringView Key) const;
+
+private:
+
+	friend class UEditorConfigSubsystem; // for access to LoadFromFile and SaveToFile
+
+	static void ReadUObject(TSharedPtr<FJsonObject> JsonObject, const UClass* Class, UObject* Instance);
+	static void ReadStruct(TSharedPtr<FJsonObject> JsonObject, const UStruct* Struct, void* Instance, UObject* Owner);
+	static void ReadValue(TSharedPtr<FJsonValue> JsonValue, const FProperty* Property, void* DataPtr, UObject* Owner);
+	
+	static TSharedPtr<FJsonObject> WriteStruct(const UStruct* Struct, const void* Instance);
+	static TSharedPtr<FJsonObject> WriteUObject(const UStruct* Struct, const UObject* Instance);
+	static TSharedPtr<FJsonValue> WriteValue(const FProperty* Property, const void* DataPtr);
+
+	static bool IsDefault(const FProperty* Property, TSharedPtr<FJsonValue> JsonValue, const void* NativeValue);
+
+	bool LoadFromFile(FStringView FilePath);
+	bool SaveToFile(FStringView FilePath) const;
+
+private:
+	TSharedPtr<UE::FJsonConfig> JsonConfig;
+	TSharedPtr<FEditorConfig> ParentConfig;
+		
+	bool Dirty { false };
+};
+
+template <typename T>
+bool FEditorConfig::TryGetStruct(FStringView Key, T& OutValue) const
+{
+	if (!IsValid())
+	{
+		return false;
+	}
+
+	TSharedPtr<FJsonObject> StructData;
+	if (!JsonConfig->TryGetJsonObject(UE::FJsonPath(Key), StructData))
+	{
+		return false;
+	}
+
+	if (!StructData.IsValid())
+	{
+		return false;
+	}
+
+	UStruct* Struct = T::StaticStruct();
+	ReadStruct(StructData, Struct, &OutValue, nullptr);
+
+	return true;
+}
+
+template <typename T>
+bool FEditorConfig::TryGetUObject(FStringView Key, T& OutValue) const
+{
+	static_assert(TIsDerivedFrom<T, UObject>::Value);
+
+	if (!IsValid())
+	{
+		return false;
+	}
+
+	TSharedPtr<FJsonObject> UObjectData;
+	if (!JsonConfig->TryGetJsonObject(UE::FJsonPath(Key), UObjectData))
+	{
+		return false;
+	}
+
+	if (!UObjectData.IsValid())
+	{
+		return false;
+	}
+
+	UClass* Class = T::StaticClass();
+	ReadUObject(UObjectData, Class, &OutValue);
+
+	return true;
+}
+
+template <typename T>
+bool FEditorConfig::TryGetRootStruct(T& OutValue) const
+{
+	if (!IsValid())
+	{
+		return false;
+	}
+
+	TSharedPtr<FJsonObject> StructData = JsonConfig->GetRootObject();
+
+	UStruct* Struct = T::StaticStruct();
+	ReadStruct(StructData, Struct, &OutValue, nullptr);
+
+	return true;
+}
+
+template <typename T>
+bool FEditorConfig::TryGetRootUObject(T& OutValue) const
+{
+	static_assert(TIsDerivedFrom<T, UObject>::Value);
+
+	if (!IsValid())
+	{
+		return false;
+	}
+
+	TSharedPtr<FJsonObject> UObjectData = JsonConfig->GetRootObject();
+
+	UClass* Class = T::StaticClass();
+	ReadUObject(UObjectData, Class, &OutValue);
+
+	return true;
+}
+
+template <typename T>
+void FEditorConfig::SetStruct(FStringView Key, const T& InValue)
+{
+	if (!IsValid())
+	{
+		return;
+	}
+
+	TSharedPtr<FJsonObject> JsonObject = WriteStruct(T::StaticStruct(), &InValue);
+	JsonConfig->SetJsonObject(UE::FJsonPath(Key), JsonObject);
+
+	Dirty = true;
+}
+
+template <typename T>
+void FEditorConfig::SetUObject(FStringView Key, const T& InValue)
+{
+	static_assert(TIsDerivedFrom<T, UObject>::Value);
+
+	if (!IsValid())
+	{
+		return;
+	}
+	
+	TSharedPtr<FJsonObject> JsonObject = WriteStruct(T::StaticClass(), &InValue);
+	JsonConfig->SetJsonObject(UE::FJsonPath(Key), JsonObject);
+		
+	Dirty = true;
+}
+
+template <typename T>
+void FEditorConfig::SetRootStruct(const T& InValue)
+{
+	if (!IsValid())
+	{
+		return;
+	}
+
+	TSharedPtr<FJsonObject> JsonObject = WriteStruct(T::StaticStruct(), &InValue);
+	JsonConfig->SetRootObject(JsonObject);
+		
+	Dirty = true;
+}
+
+template <typename T>
+void FEditorConfig::SetRootUObject(const T& InValue)
+{
+	static_assert(TIsDerivedFrom<T, UObject>::Value);
+
+	if (!IsValid())
+	{
+		return;
+	}
+	
+	TSharedPtr<FJsonObject> JsonObject = WriteStruct(T::StaticStruct(), &InValue);
+	JsonConfig->SetRootObject(JsonObject);
+		
+	Dirty = true;
+}
diff --git a/Engine/Source/Editor/EditorConfig/Public/EditorConfigModule.h b/Engine/Source/Editor/EditorConfig/Public/EditorConfigModule.h
new file mode 100644
--- /dev/null
+++ b/Engine/Source/Editor/EditorConfig/Public/EditorConfigModule.h
@@ -0,0 +1,12 @@
+// Copyright Epic Games, Inc. All Rights Reserved.
+
+#pragma once
+
+#include "Modules/ModuleInterface.h"
+
+class FEditorConfigModule : public IModuleInterface
+{
+public:
+	virtual void StartupModule() override;
+	virtual void ShutdownModule() override;
+};
diff --git a/Engine/Source/Editor/EditorConfig/Public/JsonConfig.h b/Engine/Source/Editor/EditorConfig/Public/JsonConfig.h
new file mode 100644
--- /dev/null
+++ b/Engine/Source/Editor/EditorConfig/Public/JsonConfig.h
@@ -0,0 +1,145 @@
+// Copyright Epic Games, Inc. All Rights Reserved.
+
+#pragma once
+
+#include "CoreMinimal.h"
+#include "Async/Async.h"
+#include "Dom/JsonObject.h"
+#include "Concepts/EqualityComparable.h"
+
+namespace UE
+{
+	class EDITORCONFIG_API FJsonPath
+	{
+	public:
+		struct FPart
+		{
+			FString Name;
+			int32 Index = INDEX_NONE;
+		};
+
+		FJsonPath();
+		FJsonPath(const TCHAR* Path);
+		FJsonPath(FStringView Path);
+		FJsonPath(const FJsonPath& Other);
+		FJsonPath(FJsonPath&& Other);
+
+		bool IsValid() const { return Length() > 0; }
+		const int32 Length() const { return PathParts.Num(); }
+
+		void Append(FStringView Name);
+		void SetArrayIndex(int32 Index);
+
+		FJsonPath GetSubPath(int32 NumParts) const;
+
+		FString ToString() const;
+
+		const FPart& operator[](int32 Idx) const { return PathParts[Idx]; }
+		const TArray<FPart>& GetAll() const { return PathParts; }
+
+	private:
+		void ParsePath(const FString& InPath);
+
+	private:
+		TArray<FPart> PathParts;
+	};
+
+	using FJsonValuePair = TPair<TSharedPtr<FJsonValue>, TSharedPtr<FJsonValue>>;
+
+	class EDITORCONFIG_API FJsonConfig
+	{
+	public:
+		FJsonConfig();
+
+		void SetParent(TSharedPtr<FJsonConfig> Parent);
+
+		bool LoadFromFile(FStringView FilePath);
+		bool LoadFromString(FStringView Content);
+	
+		bool SaveToFile(FStringView FilePath) const;
+		bool SaveToString(FString& OutResult) const;
+
+		bool IsValid() const { return MergedObject.IsValid(); }
+
+		const FJsonConfig* GetParentConfig() const;
+
+		template <typename T>
+		bool TryGetNumber(const FJsonPath& Path, T& OutValue) const;
+		bool TryGetBool(const FJsonPath& Path, bool& OutValue) const;
+		bool TryGetString(const FJsonPath& Path, FString& OutValue) const;
+		bool TryGetString(const FJsonPath& Path, FName& OutValue) const;
+		bool TryGetString(const FJsonPath& Path, FText& OutValue) const;
+		bool TryGetJsonValue(const FJsonPath& Path, TSharedPtr<FJsonValue>& OutValue) const;
+		bool TryGetJsonObject(const FJsonPath& Path, TSharedPtr<FJsonObject>& OutValue) const;
+		bool TryGetArray(const FJsonPath& Path, TArray<TSharedPtr<FJsonValue>>& OutArray) const;
+
+		TSharedPtr<FJsonObject> GetRootObject() const;
+
+		// these are specializations for arithmetic and string arrays
+		// these could be templated with enable-ifs, 
+		// but it ended up being more lines of incomprehensible template SFINAE than this clear list of types is
+		bool TryGetArray(const FJsonPath& Path, TArray<bool>& OutArray) const;
+		bool TryGetArray(const FJsonPath& Path, TArray<int8>& OutArray) const;
+		bool TryGetArray(const FJsonPath& Path, TArray<int16>& OutArray) const;
+		bool TryGetArray(const FJsonPath& Path, TArray<int32>& OutArray) const;
+		bool TryGetArray(const FJsonPath& Path, TArray<int64>& OutArray) const;
+		bool TryGetArray(const FJsonPath& Path, TArray<uint8>& OutArray) const;
+		bool TryGetArray(const FJsonPath& Path, TArray<uint16>& OutArray) const;
+		bool TryGetArray(const FJsonPath& Path, TArray<uint32>& OutArray) const;
+		bool TryGetArray(const FJsonPath& Path, TArray<uint64>& OutArray) const;
+		bool TryGetArray(const FJsonPath& Path, TArray<float>& OutArray) const;
+		bool TryGetArray(const FJsonPath& Path, TArray<double>& OutArray) const;
+		bool TryGetArray(const FJsonPath& Path, TArray<FString>& OutArray) const;
+		bool TryGetArray(const FJsonPath& Path, TArray<FText>& OutArray) const;
+		bool TryGetArray(const FJsonPath& Path, TArray<FName>& OutArray) const;
+
+		bool TryGetMap(const FJsonPath& Path, TArray<FJsonValuePair>& OutMap) const;
+
+		// try to set a number - returns false if:
+		// - the path references an object that doesn't exist
+		// - the path references an index in an array that doesn't exist
+		template <typename T>
+		bool SetNumber(const FJsonPath& Path, T Value);
+		bool SetBool(const FJsonPath& Path, bool Value);
+		bool SetString(const FJsonPath& Path, FStringView Value);
+		bool SetString(const FJsonPath& Path, const FText& Value);
+		bool SetJsonValue(const FJsonPath& Path, TSharedPtr<FJsonValue> Value);
+		bool SetJsonObject(const FJsonPath& Path, TSharedPtr<FJsonObject> Object);
+		bool SetJsonArray(const FJsonPath& Path, const TArray<TSharedPtr<FJsonValue>>& Array);
+
+		bool SetRootObject(TSharedPtr<FJsonObject> Object);
+
+		bool HasOverride(const FJsonPath& Path) const;
+
+	private:
+
+		template <typename T, typename TGetter>
+		bool TryGetArrayHelper(const FJsonPath& Path, TArray<T>& OutArray, TGetter Getter) const;
+
+		template <typename T>
+		bool TryGetNumericArrayHelper(const FJsonPath& Path, TArray<T>& OutArray) const;
+
+		bool SetJsonValueInMerged(const FJsonPath& Path, TSharedPtr<FJsonValue> Value);
+		bool SetJsonValueInOverride(const FJsonPath& Path, TSharedPtr<FJsonValue> NewValue, TSharedPtr<FJsonValue> PreviousValue, TSharedPtr<FJsonValue> ParentValue);
+		bool SetArrayValueInOverride(TSharedPtr<FJsonValue> CurrentValue, const TArray<TSharedPtr<FJsonValue>>& NewArray, TSharedPtr<FJsonValue> ParentValue);
+		bool SetObjectValueInOverride(TSharedPtr<FJsonObject> CurrentObject, TSharedPtr<FJsonObject> NewObject, TSharedPtr<FJsonValue> ParentValue);
+
+		bool RemoveJsonValueFromOverride(const FJsonPath& Path, TSharedPtr<FJsonValue> PreviousValue);
+
+		bool MergeThisWithParent();
+		void OnParentConfigChanged();
+
+	private:
+
+		TFuture<void> SaveFuture;
+
+		FSimpleDelegate OnConfigChanged;
+
+		TSharedPtr<FJsonConfig> ParentConfig;
+
+		TSharedPtr<FJsonObject> OverrideObject;
+		TSharedPtr<FJsonObject> MergedObject;
+	};
+}
+
+#include "JsonConfig.inl"
diff --git a/Engine/Source/Editor/EditorConfig/Public/JsonConfig.inl b/Engine/Source/Editor/EditorConfig/Public/JsonConfig.inl
new file mode 100644
--- /dev/null
+++ b/Engine/Source/Editor/EditorConfig/Public/JsonConfig.inl
@@ -0,0 +1,68 @@
+// Copyright Epic Games, Inc. All Rights Reserved.
+
+#pragma once
+
+namespace UE
+{
+	template <typename T>
+	bool FJsonConfig::TryGetNumber(const FJsonPath& Path, T& OutValue) const
+	{
+		TSharedPtr<FJsonValue> JsonValue;
+		if (!TryGetJsonValue(Path, JsonValue))
+		{
+			return false;
+		}
+
+		return JsonValue->TryGetNumber(OutValue);
+	}
+
+	template <typename T, typename TGetter>
+	bool FJsonConfig::TryGetArrayHelper(const FJsonPath& Path, TArray<T>& OutArray, TGetter Getter) const
+	{
+		if (!IsValid())
+		{
+			return false;
+		}
+
+		TArray<TSharedPtr<FJsonValue>> JsonValueArray;
+		if (!TryGetArray(Path, JsonValueArray))
+		{
+			return false;
+		}
+
+		OutArray.Reset();
+
+		for (const TSharedPtr<FJsonValue>& JsonValue : JsonValueArray)
+		{
+			T Value;
+			if (!Getter(JsonValue, Value))
+			{
+				OutArray.Reset();
+				return false;
+			}
+
+			OutArray.Add(Value);
+		}
+
+		return true;
+	}
+
+	template <typename T>
+	bool FJsonConfig::TryGetNumericArrayHelper(const FJsonPath& Path, TArray<T>& OutArray) const
+	{
+		return TryGetArrayHelper(Path, OutArray, 
+			[](const TSharedPtr<FJsonValue>& JsonValue, T& OutNumber)
+			{
+				return JsonValue->TryGetNumber(OutNumber);
+			});
+	}
+
+	template <typename T>
+	bool FJsonConfig::SetNumber(const FJsonPath& Path, T Value)
+	{
+		static_assert(TIsArithmetic<T>::Value, "Value type must be a number.");
+
+		TSharedPtr<FJsonValue> JsonValue = MakeShared<FJsonValueNumber>((double) Value);
+		return SetJsonValue(Path, JsonValue);
+	}
+}
diff --git a/Engine/Source/Runtime/Json/Private/Dom/JsonObject.cpp b/Engine/Source/Runtime/Json/Private/Dom/JsonObject.cpp
--- a/Engine/Source/Runtime/Json/Private/Dom/JsonObject.cpp
+++ b/Engine/Source/Runtime/Json/Private/Dom/JsonObject.cpp
@@ -21,6 +21,13 @@ double FJsonObject::GetNumberField( const FString& FieldName ) const
 }
 
 
+bool FJsonObject::TryGetNumberField( const FString& FieldName, float& OutNumber ) const
+{
+	TSharedPtr<FJsonValue> Field = TryGetField(FieldName);
+	return Field.IsValid() && Field->TryGetNumber(OutNumber);
+}
+
+
 bool FJsonObject::TryGetNumberField( const FString& FieldName, double& OutNumber ) const
 {
 	TSharedPtr<FJsonValue> Field = TryGetField(FieldName);
@@ -28,6 +35,20 @@ bool FJsonObject::TryGetNumberField( const FString& FieldName, double& OutNumber
 }
 
 
+bool FJsonObject::TryGetNumberField( const FString& FieldName, int8& OutNumber ) const
+{
+	TSharedPtr<FJsonValue> Field = TryGetField(FieldName);
+	return Field.IsValid() && Field->TryGetNumber(OutNumber);
+}
+
+
+bool FJsonObject::TryGetNumberField( const FString& FieldName, int16& OutNumber ) const
+{
+	TSharedPtr<FJsonValue> Field = TryGetField(FieldName);
+	return Field.IsValid() && Field->TryGetNumber(OutNumber);
+}
+
+
 bool FJsonObject::TryGetNumberField( const FString& FieldName, int32& OutNumber ) const
 {
 	TSharedPtr<FJsonValue> Field = TryGetField(FieldName);
@@ -35,18 +56,39 @@ bool FJsonObject::TryGetNumberField( const FString& FieldName, int32& OutNumber
 }
 
 
+bool FJsonObject::TryGetNumberField(const FString& FieldName, int64& OutNumber) const
+{
+	TSharedPtr<FJsonValue> Field = TryGetField(FieldName);
+	return Field.IsValid() && Field->TryGetNumber(OutNumber);
+}
+
+
+bool FJsonObject::TryGetNumberField( const FString& FieldName, uint8& OutNumber ) const
+{
+	TSharedPtr<FJsonValue> Field = TryGetField(FieldName);
+	return Field.IsValid() && Field->TryGetNumber(OutNumber);
+}
+
+
+bool FJsonObject::TryGetNumberField( const FString& FieldName, uint16& OutNumber ) const
+{
+	TSharedPtr<FJsonValue> Field = TryGetField(FieldName);
+	return Field.IsValid() && Field->TryGetNumber(OutNumber);
+}
+
 bool FJsonObject::TryGetNumberField( const FString& FieldName, uint32& OutNumber ) const
 {
 	TSharedPtr<FJsonValue> Field = TryGetField(FieldName);
 	return Field.IsValid() && Field->TryGetNumber(OutNumber);
 }
 
-bool FJsonObject::TryGetNumberField(const FString& FieldName, int64& OutNumber) const
+bool FJsonObject::TryGetNumberField( const FString& FieldName, uint64& OutNumber ) const
 {
 	TSharedPtr<FJsonValue> Field = TryGetField(FieldName);
 	return Field.IsValid() && Field->TryGetNumber(OutNumber);
 }
 
+
 void FJsonObject::SetNumberField( const FString& FieldName, double Number )
 {
 	this->Values.Add(FieldName, MakeShared<FJsonValueNumber>(Number));
@@ -166,3 +208,11 @@ void FJsonObject::SetObjectField( const FString& FieldName, const TSharedPtr<FJs
 		this->Values.Add(FieldName, MakeShared<FJsonValueNull>());
 	}
 }
+
+void FJsonObject::Duplicate(const TSharedPtr<FJsonObject>& Source, TSharedPtr<FJsonObject>& Dest)
+{
+	for (const TPair<FString, TSharedPtr<FJsonValue>>& Pair : Source->Values)
+	{
+		Dest->SetField(Pair.Key, FJsonValue::Duplicate(Pair.Value));
+	}
+}
diff --git a/Engine/Source/Runtime/Json/Private/Dom/JsonValue.cpp b/Engine/Source/Runtime/Json/Private/Dom/JsonValue.cpp
--- a/Engine/Source/Runtime/Json/Private/Dom/JsonValue.cpp
+++ b/Engine/Source/Runtime/Json/Private/Dom/JsonValue.cpp
@@ -3,9 +3,6 @@
 #include "Dom/JsonValue.h"
 #include "Dom/JsonObject.h"
 
-
-
-
 double FJsonValue::AsNumber() const
 {
 	double Number = 0.0;
@@ -269,6 +266,68 @@ bool FJsonValue::CompareEqual( const FJsonValue& Lhs, const FJsonValue& Rhs )
 	}
 }
 
+static void DuplicateJsonArray(const TArray<TSharedPtr<FJsonValue>>& Source, TArray<TSharedPtr<FJsonValue>>& Dest)
+{
+	for (const TSharedPtr<FJsonValue>& Value : Source)
+	{
+		Dest.Add(FJsonValue::Duplicate(Value));
+	} 
+}
+
+TSharedPtr<FJsonValue> FJsonValue::Duplicate(const TSharedPtr<FJsonValue>& Src)
+{
+	switch (Src->Type)
+	{
+		case EJson::Boolean:
+		{
+			bool BoolValue;
+			if (Src->TryGetBool(BoolValue))
+			{
+				return MakeShared<FJsonValueBoolean>(BoolValue);
+			}
+		}
+		case EJson::Number:
+		{
+			double NumberValue;
+			if (Src->TryGetNumber(NumberValue))
+			{
+				return MakeShared<FJsonValueNumber>(NumberValue);
+			}
+		}
+		case EJson::String:
+		{
+			FString StringValue;
+			if (Src->TryGetString(StringValue))
+			{
+				return MakeShared<FJsonValueString>(StringValue);
+			}
+		}
+		case EJson::Object:
+		{
+			const TSharedPtr<FJsonObject>* ObjectValue;
+			if (Src->TryGetObject(ObjectValue))
+			{
+				TSharedPtr<FJsonObject> NewObject = MakeShared<FJsonObject>();
+				FJsonObject::Duplicate(*ObjectValue, NewObject);
+				return MakeShared<FJsonValueObject>(NewObject);
+			}
+		}
+		case EJson::Array:
+		{
+			const TArray<TSharedPtr<FJsonValue>>* ArrayValue;
+			if (Src->TryGetArray(ArrayValue))
+			{
+				TArray<TSharedPtr<FJsonValue>> NewArray;
+				DuplicateJsonArray(*ArrayValue, NewArray);
+
+				return MakeShared<FJsonValueArray>(NewArray);
+			}
+		}
+	}
+
+	return TSharedPtr<FJsonValue>();
+}
+
 void FJsonValue::ErrorMessage(const FString& InType) const
 {
 	UE_LOG(LogJson, Error, TEXT("Json Value of type '%s' used as a '%s'."), *GetType(), *InType);
diff --git a/Engine/Source/Runtime/Json/Public/Dom/JsonObject.h b/Engine/Source/Runtime/Json/Public/Dom/JsonObject.h
--- a/Engine/Source/Runtime/Json/Public/Dom/JsonObject.h
+++ b/Engine/Source/Runtime/Json/Public/Dom/JsonObject.h
@@ -121,17 +121,35 @@ public:
 		return (int32)GetNumberField(FieldName);
 	}
 
+	/** Get the field named FieldName as a number. Returns false if it doesn't exist or cannot be converted. */
+	bool TryGetNumberField(const FString& FieldName, float& OutNumber) const;
+
 	/** Get the field named FieldName as a number. Returns false if it doesn't exist or cannot be converted. */
 	bool TryGetNumberField(const FString& FieldName, double& OutNumber) const;
 
+	/** Get the field named FieldName as a number, and makes sure it's within int8 range. Returns false if it doesn't exist or cannot be converted. */
+	bool TryGetNumberField(const FString& FieldName, int8& OutNumber) const;
+
+	/** Get the field named FieldName as a number, and makes sure it's within int16 range. Returns false if it doesn't exist or cannot be converted. */
+	bool TryGetNumberField(const FString& FieldName, int16& OutNumber) const;
+
 	/** Get the field named FieldName as a number, and makes sure it's within int32 range. Returns false if it doesn't exist or cannot be converted. */
 	bool TryGetNumberField(const FString& FieldName, int32& OutNumber) const;
 
+	/** Get the field named FieldName as a number. Returns false if it doesn't exist or cannot be converted. */
+	bool TryGetNumberField(const FString& FieldName, int64& OutNumber) const;
+
+	/** Get the field named FieldName as a number, and makes sure it's within uint8 range. Returns false if it doesn't exist or cannot be converted. */
+	bool TryGetNumberField(const FString& FieldName, uint8& OutNumber) const;
+
+	/** Get the field named FieldName as a number, and makes sure it's within uint16 range. Returns false if it doesn't exist or cannot be converted. */
+	bool TryGetNumberField(const FString& FieldName, uint16& OutNumber) const;
+
 	/** Get the field named FieldName as a number, and makes sure it's within uint32 range. Returns false if it doesn't exist or cannot be converted.  */
 	bool TryGetNumberField(const FString& FieldName, uint32& OutNumber) const;
 
-	/** Get the field named FieldName as a number. Returns false if it doesn't exist or cannot be converted. */
-	bool TryGetNumberField(const FString& FieldName, int64& OutNumber) const;
+	/** Get the field named FieldName as a number, and makes sure it's within uint64 range. Returns false if it doesn't exist or cannot be converted.  */
+	bool TryGetNumberField(const FString& FieldName, uint64& OutNumber) const;
 
 	/** Add a field named FieldName with Number as value */
 	void SetNumberField( const FString& FieldName, double Number );
@@ -210,4 +228,6 @@ public:
 
 	/** Set an ObjectField named FieldName and value of JsonObject */
 	void SetObjectField( const FString& FieldName, const TSharedPtr<FJsonObject>& JsonObject );
+
+	static void Duplicate(const TSharedPtr<FJsonObject>& Source, TSharedPtr<FJsonObject>& Dest);
 };
diff --git a/Engine/Source/Runtime/Json/Public/Dom/JsonValue.h b/Engine/Source/Runtime/Json/Public/Dom/JsonValue.h
--- a/Engine/Source/Runtime/Json/Public/Dom/JsonValue.h
+++ b/Engine/Source/Runtime/Json/Public/Dom/JsonValue.h
@@ -68,10 +68,16 @@ public:
 
 	/** Tries to convert this value to an array, returning false if not possible */
 	virtual bool TryGetArray(const TArray< TSharedPtr<FJsonValue> >*& OutArray) const { return false; }
+	
+	/** Tries to convert this value to an array, returning false if not possible */
+	virtual bool TryGetArray(TArray< TSharedPtr<FJsonValue> >*& OutArray) { return false; }
 
 	/** Tries to convert this value to an object, returning false if not possible */
 	virtual bool TryGetObject(const TSharedPtr<FJsonObject>*& Object) const { return false; }
 
+	/** Tries to convert this value to an object, returning false if not possible */
+	virtual bool TryGetObject(TSharedPtr<FJsonObject>*& Object) { return false; }
+
 	/** Returns true if this value is a 'null' */
 	bool IsNull() const { return Type == EJson::Null || Type == EJson::None; }
 
@@ -84,6 +90,8 @@ public:
 
 	EJson Type;
 
+	static TSharedPtr<FJsonValue> Duplicate(const TSharedPtr<FJsonValue>& Src);
+
 	static bool CompareEqual(const FJsonValue& Lhs, const FJsonValue& Rhs);
 
 protected:
@@ -196,6 +204,7 @@ class JSON_API FJsonValueArray : public FJsonValue
 public:
 	FJsonValueArray(const TArray< TSharedPtr<FJsonValue> >& InArray) : Value(InArray) {Type = EJson::Array;}
 	virtual bool TryGetArray(const TArray< TSharedPtr<FJsonValue> >*& OutArray) const override	{ OutArray = &Value; return true; }
+	virtual bool TryGetArray(TArray< TSharedPtr<FJsonValue> >*& OutArray) override				{ OutArray = &Value; return true; }
 	
 protected:
 	TArray< TSharedPtr<FJsonValue> > Value;
@@ -209,7 +218,8 @@ class JSON_API FJsonValueObject : public FJsonValue
 {
 public:
 	FJsonValueObject(TSharedPtr<FJsonObject> InObject) : Value(InObject) {Type = EJson::Object;}
-	virtual bool TryGetObject(const TSharedPtr<FJsonObject>*& OutObject) const override			{ OutObject = &Value; return true; }
+	virtual bool TryGetObject(const TSharedPtr<FJsonObject>*& OutObject) const override	{ OutObject = &Value; return true; }
+	virtual bool TryGetObject(TSharedPtr<FJsonObject>*& OutObject) override				{ OutObject = &Value; return true; }
 	
 protected:
 	TSharedPtr<FJsonObject> Value;
