Added EditorMetadataOverrides subsystem and enabled its usage for details view favourites.

The Favorites category is now always visible, and displays a note about how to add properties there when empty.

diff --git a/Engine/Source/Editor/EditorConfig/Private/EditorConfigSubystem.cpp b/Engine/Source/Editor/EditorConfig/Private/EditorConfigSubystem.cpp
new file mode 100644
--- /dev/null
+++ b/Engine/Source/Editor/EditorConfig/Private/EditorConfigSubystem.cpp
@@ -0,0 +1,133 @@
+// Copyright Epic Games, Inc. All Rights Reserved.
+
+#include "EditorConfigSubsystem.h"
+
+#include "Async/Async.h"
+#include "Misc/ScopeExit.h"
+
+UEditorConfigSubsystem::UEditorConfigSubsystem()
+{
+	
+}
+
+void UEditorConfigSubsystem::Initialize(FSubsystemCollectionBase& Collection)
+{
+	AddSearchDirectory(FPaths::Combine(FPaths::Combine(FPlatformProcess::UserSettingsDir(), *FApp::GetEpicProductIdentifier(), TEXT("Editor")))); // AppData
+	AddSearchDirectory(FPaths::Combine(FPaths::EngineConfigDir(), TEXT("Editor"))); // Engine
+	AddSearchDirectory(FPaths::Combine(FPaths::ProjectConfigDir(), TEXT("Editor"))); // ProjectName
+}
+
+TSharedPtr<FEditorConfig> UEditorConfigSubsystem::FindOrLoadConfig(FStringView ConfigName)
+{
+	FString ConfigNameString(ConfigName);
+		
+	TSharedPtr<FEditorConfig> Current;
+
+	for (const FString& Dir : SearchDirectories)
+	{
+		FString FullPath = FPaths::Combine(Dir, ConfigNameString) + TEXT(".json");
+
+		const TSharedPtr<FEditorConfig>* Existing = LoadedConfigs.Find(FullPath);
+		if (Existing != nullptr && Existing->IsValid())
+		{
+			Current = *Existing;
+		}
+		else
+		{
+			// didn't exist yet, load now
+			TSharedPtr<FEditorConfig> NewConfig = MakeShared<FEditorConfig>();
+			if (NewConfig->LoadFromFile(FullPath))
+			{
+				LoadedConfigs.Add(FullPath, NewConfig);
+
+				if (Current.IsValid())
+				{
+					NewConfig->SetParent(Current);
+				}
+
+				Current = NewConfig;
+			}
+		}
+	}
+
+	FString FinalPath = FPaths::Combine(SearchDirectories.Last(), ConfigNameString) + TEXT(".json");
+
+	const TSharedPtr<FEditorConfig>* FinalConfig = LoadedConfigs.Find(FinalPath);
+	if (FinalConfig != nullptr)
+	{
+		return *FinalConfig;
+	}
+
+	// no config in the last search directory, create one now
+	// this will be the config that changes are written to
+	TSharedPtr<FEditorConfig> NewConfig = MakeShared<FEditorConfig>();
+	if (Current.IsValid())
+	{
+		// parent to the previous config
+		NewConfig->SetParent(Current);
+	}
+
+	LoadedConfigs.Add(FinalPath, NewConfig);
+
+	return NewConfig;
+}
+
+void UEditorConfigSubsystem::SaveConfig(TSharedPtr<FEditorConfig> Config, FOnCompletedDelegate OnCompleted)
+{
+	const FString* FilePath = LoadedConfigs.FindKey(Config);
+	if (FilePath == nullptr)
+	{
+		return;
+	}
+
+	SaveLock.WriteLock();
+	ON_SCOPE_EXIT { SaveLock.WriteUnlock(); };
+
+	for (const FPendingSave& Save : PendingSaves)
+	{
+		if (Save.Config == Config)
+		{
+			// save already pending
+			return;
+		}
+	}
+
+	FString File = *FilePath;
+
+	FPendingSave& Save = PendingSaves.AddDefaulted_GetRef();
+	Save.Config = Config;
+	Save.FileName = File;
+	Save.OnCompleted = OnCompleted;
+	Save.WasSuccess = Async(EAsyncExecution::Thread, [Config, File]()
+		{
+			return Config->SaveToFile(File);
+		}, 
+		[this, Config]() 
+		{
+			OnSaveCompleted(Config);
+		});
+}
+
+void UEditorConfigSubsystem::OnSaveCompleted(TSharedPtr<FEditorConfig> Config)
+{
+	SaveLock.WriteLock();
+	ON_SCOPE_EXIT { SaveLock.WriteUnlock(); };
+
+	int32 Index = PendingSaves.IndexOfByPredicate([Config](const FPendingSave& Element)
+		{
+			return Element.Config == Config;
+		});
+
+	if (Index != INDEX_NONE)
+	{
+		const FPendingSave& PendingSave = PendingSaves[Index];
+		PendingSave.OnCompleted.ExecuteIfBound(PendingSave.WasSuccess.Get());
+
+		PendingSaves.RemoveAt(Index);
+	}
+}
+
+void UEditorConfigSubsystem::AddSearchDirectory(FStringView SearchDir)
+{
+	SearchDirectories.AddUnique(FString(SearchDir));
+}
\ No newline at end of file
diff --git a/Engine/Source/Editor/EditorConfig/Private/EditorMetadataOverrides.cpp b/Engine/Source/Editor/EditorConfig/Private/EditorMetadataOverrides.cpp
new file mode 100644
--- /dev/null
+++ b/Engine/Source/Editor/EditorConfig/Private/EditorMetadataOverrides.cpp
@@ -0,0 +1,352 @@
+// Copyright Epic Games, Inc. All Rights Reserved.
+
+#include "EditorMetadataOverrides.h"
+
+#include "Editor.h"
+#include "EditorConfigSubsystem.h"
+#include "Misc/StringBuilder.h"
+#include "UObject/UnrealType.h"
+
+UEditorMetadataOverrides::UEditorMetadataOverrides()
+{
+}
+
+void UEditorMetadataOverrides::Initialize(FSubsystemCollectionBase& Collection)
+{
+	UEditorConfigSubsystem* EditorConfig = Collection.InitializeDependency<UEditorConfigSubsystem>();
+	
+	TSharedPtr<FEditorConfig> MetadataOverrideConfig = EditorConfig->FindOrLoadConfig(TEXT("MetadataOverrides"));
+	LoadFromConfig(MetadataOverrideConfig);
+}
+
+bool UEditorMetadataOverrides::LoadFromConfig(TSharedPtr<FEditorConfig> Config)
+{
+	SourceConfig = Config;
+	LoadedMetadata.Classes.Reset();
+
+	if (!SourceConfig.IsValid())
+	{
+		return false;
+	}
+
+	return SourceConfig->TryGetStruct(TEXT("Metadata"), LoadedMetadata);
+}
+
+void UEditorMetadataOverrides::Tick(float DeltaTime)
+{
+	if (bDirty)
+	{
+		TimeSinceLastSave = TimeSinceLastSave + DeltaTime;
+
+		const float SaveDelaySeconds = 3.0f;
+		if (TimeSinceLastSave > SaveDelaySeconds)
+		{
+			Save();
+		}
+	}
+}
+
+TStatId UEditorMetadataOverrides::GetStatId() const
+{
+	RETURN_QUICK_DECLARE_CYCLE_STAT(UEditorMetadataOverrides, STATGROUP_Tickables);
+}
+
+void UEditorMetadataOverrides::Save()
+{
+	if (!SourceConfig.IsValid())
+	{
+		return;
+	}
+
+	SourceConfig->SetStruct(TEXT("Metadata"), LoadedMetadata);
+
+	UEditorConfigSubsystem* EditorConfigSubsystem = GEditor->GetEditorSubsystem<UEditorConfigSubsystem>();
+	if (EditorConfigSubsystem == nullptr)
+	{
+		return;
+	}
+
+	EditorConfigSubsystem->SaveConfig(SourceConfig, FOnCompletedDelegate::CreateUObject(this, &UEditorMetadataOverrides::OnCompleted));
+}
+
+void UEditorMetadataOverrides::OnCompleted(bool bSuccess)
+{
+	if (bSuccess)
+	{
+		TimeSinceLastSave = 0;
+		bDirty = false;
+	}
+}
+
+const FMetadataSet* UEditorMetadataOverrides::FindFieldMetadata(const FField* Field, FName Key) const
+{
+	if (!SourceConfig.IsValid())
+	{
+		return nullptr;
+	}
+
+	const UStruct* OwnerStruct = Field->GetOwnerStruct();
+	if (OwnerStruct == nullptr)
+	{
+		return nullptr;
+	}
+
+	const FStructMetadata* StructMetadata = LoadedMetadata.Classes.Find(OwnerStruct->GetFName());
+	if (StructMetadata == nullptr)
+	{
+		return nullptr;
+	}
+
+	const FMetadataSet* FieldMetadata = StructMetadata->Fields.Find(Field->GetFName());
+	return FieldMetadata;
+}
+	
+FMetadataSet* UEditorMetadataOverrides::FindOrAddFieldMetadata(const FField* Field, FName Key)
+{
+	if (!SourceConfig.IsValid())
+	{
+		return nullptr;
+	}
+
+	const UStruct* OwnerStruct = Field->GetOwnerStruct();
+	if (OwnerStruct == nullptr)
+	{
+		return nullptr;
+	}
+
+	FStructMetadata& StructMetadata = LoadedMetadata.Classes.FindOrAdd(OwnerStruct->GetFName());
+
+	FMetadataSet& FieldMetadata = StructMetadata.Fields.FindOrAdd(Field->GetFName());
+	return &FieldMetadata;
+}
+
+bool UEditorMetadataOverrides::HasMetadata(const FField* Field, FName Key) const
+{
+	check(Field != nullptr);
+
+	const FMetadataSet* FieldMetadata = FindFieldMetadata(Field, Key);
+	return FieldMetadata != nullptr;
+}
+
+bool UEditorMetadataOverrides::GetStringMetadata(const FField* Field, FName Key, FString& OutValue) const
+{
+	check(Field != nullptr);
+
+	const FMetadataSet* FieldMetadata = FindFieldMetadata(Field, Key);
+	if (FieldMetadata == nullptr)
+	{
+		return false;
+	}
+
+	const FString* MetaValue = FieldMetadata->Strings.Find(Key);
+	if (MetaValue == nullptr)
+	{
+		if (!Field->HasMetaData(Key))
+		{
+			return false;
+		}
+
+		OutValue = Field->GetMetaData(Key);
+		return true;
+	}
+
+	OutValue = *MetaValue;
+	return true;
+}
+
+void UEditorMetadataOverrides::SetStringMetadata(const FField* Field, FName Key, FStringView Value)
+{
+	check(Field != nullptr);
+		
+	FMetadataSet* FieldMetadata = FindOrAddFieldMetadata(Field, Key);
+	if (FieldMetadata == nullptr)
+	{
+		return;
+	}
+
+	FieldMetadata->Strings.Add(Key, FString(Value));
+	bDirty = true;
+}
+
+bool UEditorMetadataOverrides::GetFloatMetadata(const FField* Field, FName Key, float& OutValue) const
+{
+check(Field != nullptr);
+
+	const FMetadataSet* FieldMetadata = FindFieldMetadata(Field, Key);
+	if (FieldMetadata == nullptr)
+	{
+		return false;
+	}
+
+	const float* MetaValue = FieldMetadata->Floats.Find(Key);
+	if (MetaValue == nullptr)
+	{
+		if (!Field->HasMetaData(Key))
+		{
+			return false;
+		}
+
+		OutValue = Field->GetFloatMetaData(Key);
+		return true;
+	}
+
+	OutValue = *MetaValue;
+	return true;
+}
+
+void UEditorMetadataOverrides::SetFloatMetadata(const FField* Field, FName Key, float Value)
+{
+	check(Field != nullptr);
+		
+	FMetadataSet* FieldMetadata = FindOrAddFieldMetadata(Field, Key);
+	if (FieldMetadata == nullptr)
+	{
+		return;
+	}
+
+	FieldMetadata->Floats.Add(Key, Value);
+	bDirty = true;
+}
+
+bool UEditorMetadataOverrides::GetIntMetadata(const FField* Field, FName Key, int32& OutValue) const
+{
+	check(Field != nullptr);
+
+	const FMetadataSet* FieldMetadata = FindFieldMetadata(Field, Key);
+	if (FieldMetadata == nullptr)
+	{
+		return false;
+	}
+
+	const int32* MetaValue = FieldMetadata->Ints.Find(Key);
+	if (MetaValue == nullptr)
+	{
+		if (!Field->HasMetaData(Key))
+		{
+			return false;
+		}
+
+		OutValue = Field->GetIntMetaData(Key);
+		return true;
+	}
+
+	OutValue = *MetaValue;
+	return true;
+}
+
+void UEditorMetadataOverrides::SetIntMetadata(const FField* Field, FName Key, int32 Value)
+{
+	check(Field != nullptr);
+		
+	FMetadataSet* FieldMetadata = FindOrAddFieldMetadata(Field, Key);
+	if (FieldMetadata == nullptr)
+	{
+		return;
+	}
+
+	FieldMetadata->Ints.Add(Key, Value);
+	bDirty = true;
+}
+
+bool UEditorMetadataOverrides::GetBoolMetadata(const FField* Field, FName Key, bool& OutValue) const
+{
+	check(Field != nullptr);
+
+	const FMetadataSet* FieldMetadata = FindFieldMetadata(Field, Key);
+	if (FieldMetadata == nullptr)
+	{
+		return false;
+	}
+
+	const bool* MetaValue = FieldMetadata->Bools.Find(Key);
+	if (MetaValue == nullptr)
+	{
+		if (!Field->HasMetaData(Key))
+		{
+			return false;
+		}
+
+		OutValue = Field->GetBoolMetaData(Key);
+		return true;
+	}
+
+	OutValue = *MetaValue;
+	return true;
+}
+	
+void UEditorMetadataOverrides::SetBoolMetadata(const FField* Field, FName Key, bool Value)
+{
+	check(Field != nullptr);
+		
+	FMetadataSet* FieldMetadata = FindOrAddFieldMetadata(Field, Key);
+	if (FieldMetadata == nullptr)
+	{
+		return;
+	}
+
+	FieldMetadata->Bools.Add(Key, Value);
+	bDirty = true;
+}
+
+bool UEditorMetadataOverrides::GetClassMetadata(const FField* Field, FName Key, UClass*& OutValue) const
+{
+	check(Field != nullptr);
+
+	FString ClassName;
+	if (!GetStringMetadata(Field, Key, ClassName))
+	{
+		return false;
+	}
+
+	OutValue = FindObject<UClass>(ANY_PACKAGE, *ClassName);
+	return true; // we return true here because we did have a value, it just wasn't a valid class name
+}
+
+void UEditorMetadataOverrides::SetClassMetadata(const FField* Field, FName Key, UClass* Value)
+{
+	check(Field != nullptr);
+
+	FString ClassName;
+	if (Value != nullptr)
+	{
+		ClassName = Value->GetName();
+	}
+
+	SetStringMetadata(Field, Key, ClassName);
+	bDirty = true;
+}
+
+void UEditorMetadataOverrides::RemoveMetadata(const FField* Field, FName Key)
+{
+	if (!SourceConfig.IsValid())
+	{
+		return;
+	}
+
+	const UStruct* OwnerStruct = Field->GetOwnerStruct();
+	if (OwnerStruct == nullptr)
+	{
+		return;
+	}
+
+	FStructMetadata* StructMetadata = LoadedMetadata.Classes.Find(OwnerStruct->GetFName());
+	if (StructMetadata == nullptr)
+	{
+		return;
+	}
+
+	FMetadataSet* FieldMetadata = StructMetadata->Fields.Find(Field->GetFName());
+	if (FieldMetadata == nullptr)
+	{
+		return;
+	}
+
+	FieldMetadata->Ints.Remove(Key);
+	FieldMetadata->Bools.Remove(Key);
+	FieldMetadata->Floats.Remove(Key);
+	FieldMetadata->Strings.Remove(Key);
+
+	bDirty = true;
+}
+
+PRAGMA_ENABLE_OPTIMIZATION
\ No newline at end of file
diff --git a/Engine/Source/Editor/EditorConfig/Public/EditorConfigSubsystem.h b/Engine/Source/Editor/EditorConfig/Public/EditorConfigSubsystem.h
new file mode 100644
--- /dev/null
+++ b/Engine/Source/Editor/EditorConfig/Public/EditorConfigSubsystem.h
@@ -0,0 +1,44 @@
+// Copyright Epic Games, Inc. All Rights Reserved.
+
+#pragma once
+
+#include "EditorConfig.h"
+#include "EditorSubsystem.h"
+
+#include "EditorConfigSubsystem.generated.h"
+
+DECLARE_DELEGATE_OneParam(FOnCompletedDelegate, bool);
+
+UCLASS()
+class EDITORCONFIG_API UEditorConfigSubsystem : 
+	public UEditorSubsystem
+{
+	GENERATED_BODY()
+
+public:
+	UEditorConfigSubsystem();
+
+	void Initialize(FSubsystemCollectionBase& Collection);
+
+	TSharedPtr<FEditorConfig> FindOrLoadConfig(FStringView ConfigName);
+	void SaveConfig(TSharedPtr<FEditorConfig> Config, FOnCompletedDelegate OnCompleted);
+
+	void AddSearchDirectory(FStringView SearchDir);
+
+private:
+	void OnSaveCompleted(TSharedPtr<FEditorConfig> Config);
+
+private:
+	struct FPendingSave
+	{
+		FString FileName;
+		TSharedPtr<FEditorConfig> Config;
+		TFuture<bool> WasSuccess;
+		FOnCompletedDelegate OnCompleted;
+	};
+
+	FRWLock SaveLock;
+	TArray<FPendingSave> PendingSaves;
+	TArray<FString> SearchDirectories;
+	TMap<FString, TSharedPtr<FEditorConfig>> LoadedConfigs;
+};
diff --git a/Engine/Source/Editor/EditorConfig/Public/EditorMetadataOverrides.h b/Engine/Source/Editor/EditorConfig/Public/EditorMetadataOverrides.h
new file mode 100644
--- /dev/null
+++ b/Engine/Source/Editor/EditorConfig/Public/EditorMetadataOverrides.h
@@ -0,0 +1,104 @@
+// Copyright Epic Games, Inc. All Rights Reserved.
+
+#pragma once
+
+#include "CoreMinimal.h"
+#include "EditorConfig.h"
+#include "EditorSubsystem.h"
+#include "TickableEditorObject.h"
+
+#include "EditorMetadataOverrides.generated.h"
+
+USTRUCT()
+struct FMetadataSet
+{
+	GENERATED_BODY()
+
+	// map of metadata key to metadata value
+	UPROPERTY()
+	TMap<FName, FString> Strings;
+		
+	UPROPERTY()
+	TMap<FName, bool> Bools;
+
+	UPROPERTY()
+	TMap<FName, int32> Ints;
+		
+	UPROPERTY()
+	TMap<FName, float> Floats;
+};
+
+USTRUCT()
+struct FStructMetadata
+{
+	GENERATED_BODY()
+
+	// map of field name to field metadata
+	UPROPERTY()
+	TMap<FName, FMetadataSet> Fields;
+
+	UPROPERTY()
+	FMetadataSet StructMetadata;
+};
+
+USTRUCT()
+struct FMetadataConfig
+{
+	GENERATED_BODY()
+		
+	// map of class name to class metadata
+	UPROPERTY()
+	TMap<FName, FStructMetadata> Classes;
+};
+
+UCLASS()
+class EDITORCONFIG_API UEditorMetadataOverrides : 
+	public UEditorSubsystem, 
+	public FTickableEditorObject
+{ 
+	GENERATED_BODY()
+
+public:
+	UEditorMetadataOverrides();
+	virtual ~UEditorMetadataOverrides() {}
+
+	virtual void Initialize(FSubsystemCollectionBase& Collection) override;
+
+	bool LoadFromConfig(TSharedPtr<FEditorConfig> Config);
+	void Save();
+
+	bool HasMetadata(const FField* Field, FName Key) const;
+
+	bool GetStringMetadata(const FField* Field, FName Key, FString& OutValue) const;
+	void SetStringMetadata(const FField* Field, FName Key, FStringView Value);
+
+	bool GetFloatMetadata(const FField* Field, FName Key, float& OutValue) const;
+	void SetFloatMetadata(const FField* Field, FName Key, float Value);
+
+	bool GetIntMetadata(const FField* Field, FName Key, int32& OutValue) const;
+	void SetIntMetadata(const FField* Field, FName Key, int32 Value);
+
+	bool GetBoolMetadata(const FField* Field, FName Key, bool& OutValue) const;
+	void SetBoolMetadata(const FField* Field, FName Key, bool Value);
+
+	bool GetClassMetadata(const FField* Field, FName Key, UClass*& OutValue) const;
+	void SetClassMetadata(const FField* Field, FName Key, UClass* Value);
+
+	void RemoveMetadata(const FField* Field, FName Key);
+
+	void Tick(float DeltaTime) override;
+	TStatId GetStatId() const override;
+
+private:
+	const FMetadataSet* FindFieldMetadata(const FField* Field, FName Key) const;
+	FMetadataSet* FindOrAddFieldMetadata(const FField* Field, FName Key);
+
+	void OnCompleted(bool bSuccess);
+
+private:
+	TSharedPtr<FEditorConfig> SourceConfig;
+	FMetadataConfig LoadedMetadata;
+
+	TAtomic<bool> bDirty { false };
+	TAtomic<float> TimeSinceLastSave { 0 };
+};
diff --git a/Engine/Source/Editor/PropertyEditor/Private/DetailCategoryBuilderImpl.cpp b/Engine/Source/Editor/PropertyEditor/Private/DetailCategoryBuilderImpl.cpp
--- a/Engine/Source/Editor/PropertyEditor/Private/DetailCategoryBuilderImpl.cpp
+++ b/Engine/Source/Editor/PropertyEditor/Private/DetailCategoryBuilderImpl.cpp
@@ -180,6 +180,23 @@ IDetailGroup& FDetailCategoryImpl::AddGroup(FName GroupName, const FText& Locali
 	return *NewCustomization.DetailGroup;
 }
 
+int32 FDetailCategoryImpl::GetNumCustomizations() const
+{
+	int32 NumCustomizations = 0;
+
+	for (int32 LayoutIndex = 0; LayoutIndex < LayoutMap.Num(); ++LayoutIndex)
+	{
+		const FDetailLayout& Layout = LayoutMap[LayoutIndex];
+
+		NumCustomizations += Layout.GetDefaultSimpleLayouts().Num();
+		NumCustomizations += Layout.GetDefaultAdvancedLayouts().Num();
+		NumCustomizations += Layout.GetCustomSimpleLayouts().Num();
+		NumCustomizations += Layout.GetCustomAdvancedLayouts().Num();
+	}
+
+	return NumCustomizations;
+}
+
 void FDetailCategoryImpl::GetDefaultProperties(TArray<TSharedRef<IPropertyHandle> >& OutDefaultProperties, bool bSimpleProperties, bool bAdvancedProperties)
 {
 	FDetailLayoutBuilderImpl& DetailLayoutBuilderRef = GetParentLayoutImpl();
@@ -611,12 +628,6 @@ bool FDetailCategoryImpl::ContainsOnlyAdvanced() const
 	return !bFavoriteCategory && SimpleChildNodes.Num() == 0 && AdvancedChildNodes.Num() > 0;
 }
 
-void FDetailCategoryImpl::GetCategoryInformation(int32 &SimpleChildNum, int32 &AdvanceChildNum) const
-{
-	SimpleChildNum = SimpleChildNodes.Num();
-	AdvanceChildNum = AdvancedChildNodes.Num();
-}
-
 void FDetailCategoryImpl::SetDisplayName(FName InCategoryName, const FText& LocalizedNameOverride)
 {
 	if (!LocalizedNameOverride.IsEmpty())
diff --git a/Engine/Source/Editor/PropertyEditor/Private/DetailCategoryBuilderImpl.h b/Engine/Source/Editor/PropertyEditor/Private/DetailCategoryBuilderImpl.h
--- a/Engine/Source/Editor/PropertyEditor/Private/DetailCategoryBuilderImpl.h
+++ b/Engine/Source/Editor/PropertyEditor/Private/DetailCategoryBuilderImpl.h
@@ -19,8 +19,10 @@
 #include "DetailCategoryBuilder.h"
 #include "PropertyCustomizationHelpers.h"
 
+class FDetailCustomBuilderRow;
 class FDetailGroup;
 class FDetailPropertyRow;
+class FDetailWidgetRow;
 class IDetailGroup;
 class IDetailPropertyRow;
 
@@ -31,13 +33,13 @@ struct FDetailLayoutCustomization
 {
 	FDetailLayoutCustomization();
 	/** The property node for the property detail */
-	TSharedPtr<class FDetailPropertyRow> PropertyRow;
+	TSharedPtr<FDetailPropertyRow> PropertyRow;
 	/** A group of customizations */
-	TSharedPtr<class FDetailGroup> DetailGroup;
+	TSharedPtr<FDetailGroup> DetailGroup;
 	/** Custom Widget for displaying the detail */
-	TSharedPtr<class FDetailWidgetRow> WidgetDecl;
+	TSharedPtr<FDetailWidgetRow> WidgetDecl;
 	/** Custom builder for more complicated widgets */
-	TSharedPtr<class FDetailCustomBuilderRow> CustomBuilderRow;
+	TSharedPtr<FDetailCustomBuilderRow> CustomBuilderRow;
 	/** @return true if this customization has a property node */
 	bool HasPropertyNode() const { return GetPropertyNode().IsValid(); }
 	/** @return true if this customization has a custom widget */
@@ -291,8 +293,7 @@ public:
 	/** @return true if this category only contains advanced properties */
 	bool ContainsOnlyAdvanced() const;
 
-	/** @return true if this category only contains advanced properties */
-	void GetCategoryInformation(int32 &SimpleChildNum, int32 &AdvanceChildNum) const;
+	int32 GetNumCustomizations() const;
 
 	/**
 	 * Called when the advanced dropdown button is clicked
@@ -305,6 +306,7 @@ public:
 	void SetCategoryAsSpecialFavorite() { bFavoriteCategory = true; bForceAdvanced = true; }
 
 	FDetailLayoutCustomization* GetDefaultCustomization(TSharedRef<FPropertyNode> PropertyNode);
+
 private:
 	virtual void OnItemExpansionChanged(bool bIsExpanded, bool bShouldSaveState) override;
 
diff --git a/Engine/Source/Editor/PropertyEditor/Private/DetailLayoutHelpers.cpp b/Engine/Source/Editor/PropertyEditor/Private/DetailLayoutHelpers.cpp
--- a/Engine/Source/Editor/PropertyEditor/Private/DetailLayoutHelpers.cpp
+++ b/Engine/Source/Editor/PropertyEditor/Private/DetailLayoutHelpers.cpp
@@ -14,7 +14,6 @@
 
 namespace DetailLayoutHelpers
 {
-
 	void UpdateSinglePropertyMapRecursive(FPropertyNode& InNode, FName CurCategory, FComplexPropertyNode* CurObjectNode, FUpdatePropertyMapArgs& InUpdateArgs)
 	{
 		FDetailLayoutData& LayoutData = *InUpdateArgs.LayoutData;
@@ -160,16 +159,17 @@ namespace DetailLayoutHelpers
 
 						PropertyNodeMap.Add(Property->GetFName(), ChildNodePtr);
 					}
+
 					bool bCanDisplayFavorite = false;
 					if (bVisibleByDefault && bIsUserVisible && !bPushOutStructProps)
 					{
 						FName CategoryName = CurCategory;
 						// For properties inside a struct, add them to their own category unless they just take the name of the parent struct.  
 						// In that case push them to the parent category
-						FName PropertyCatagoryName = FObjectEditorUtils::GetCategoryFName(Property);
-						if (!ParentStructProp || (PropertyCatagoryName != ParentStructProp->Struct->GetFName()))
+						FName PropertyCategoryName = FObjectEditorUtils::GetCategoryFName(Property);
+						if (!ParentStructProp || (PropertyCategoryName != ParentStructProp->Struct->GetFName()))
 						{
-							CategoryName = PropertyCatagoryName;
+							CategoryName = PropertyCategoryName;
 						}
 
 						if (!LocalUpdateFavoriteSystemOnly)
@@ -190,25 +190,14 @@ namespace DetailLayoutHelpers
 							if (bIsCustomizedStruct)
 							{
 								bCanDisplayFavorite = false;
-								//CustomizedStruct child are not categorize since they are under an object but we have to put them in favorite category if the user want to favorite them
+								//CustomizedStruct child are not categorized since they are under an object but we have to put them in favorite category if the user wants to favorite them
 								LocalUpdateFavoriteSystemOnly = true;
 							}
 							else if (ChildNodePtr->IsFavorite())
 							{
-								//Find or create the favorite category, we have to duplicate favorite property row under this category
-								FString CategoryFavoritesName = TEXT("Favorites");
-								FName CatFavName = *CategoryFavoritesName;
-								FDetailCategoryImpl& CategoryFavImpl = DetailLayout.DefaultCategory(CatFavName);
-								CategoryFavImpl.SetSortOrder(0);
-								CategoryFavImpl.SetCategoryAsSpecialFavorite();
-
-								//Add the property to the favorite
-								FObjectPropertyNode *RootObjectParent = ChildNodePtr->FindRootObjectItemParent();
-								FName RootInstanceName = NAME_None;
-								if (RootObjectParent != nullptr)
-								{
-									RootInstanceName = RootObjectParent->GetObjectBaseClass()->GetFName();
-								}
+								// Find or create the favorite category, we have to duplicate favorite property row under this category
+								static const FName FavoritesCategoryName(TEXT("Favorites"));
+								FDetailCategoryImpl& FavoritesCategory = DetailLayout.DefaultCategory(FavoritesCategoryName);
 
 								if (LocalUpdateFavoriteSystemOnly)
 								{
@@ -218,9 +207,9 @@ namespace DetailLayoutHelpers
 									}
 									else
 									{
-										//If the parent has a condition that is not met, make the child as readonly
+										//If the parent has a condition that is not met, mark the child as read-only
 										FDetailLayoutCustomization ParentTmpCustomization;
-										ParentTmpCustomization.PropertyRow = MakeShareable(new FDetailPropertyRow(InNode.AsShared(), CategoryFavImpl.AsShared()));
+										ParentTmpCustomization.PropertyRow = MakeShareable(new FDetailPropertyRow(InNode.AsShared(), FavoritesCategory.AsShared()));
 										if (ParentTmpCustomization.PropertyRow->GetPropertyEditor()->IsPropertyEditingEnabled() == false)
 										{
 											ChildNode.SetNodeFlags(EPropertyNodeFlags::IsReadOnly, true);
@@ -228,14 +217,24 @@ namespace DetailLayoutHelpers
 									}
 								}
 
+								//Add the property to the favorite
+								FObjectPropertyNode* RootObjectParent = ChildNodePtr->FindRootObjectItemParent();
+								FName RootInstanceName = NAME_None;
+								if (RootObjectParent != nullptr)
+								{
+									RootInstanceName = RootObjectParent->GetObjectBaseClass()->GetFName();
+								}
+
 								//Duplicate the row
-								CategoryFavImpl.AddPropertyNode(ChildNodePtr.ToSharedRef(), RootInstanceName);
+								FavoritesCategory.AddPropertyNode(ChildNodePtr.ToSharedRef(), RootInstanceName);
 							}
 
 							if (bIsStruct)
 							{
 								LocalUpdateFavoriteSystemOnly = true;
 							}
+
+							ChildArgs.bUpdateFavoriteSystemOnly = LocalUpdateFavoriteSystemOnly;
 						}
 					}
 					ChildNodePtr->SetCanDisplayFavorite(bCanDisplayFavorite);
diff --git a/Engine/Source/Editor/PropertyEditor/Private/IDetailsViewPrivate.h b/Engine/Source/Editor/PropertyEditor/Private/IDetailsViewPrivate.h
--- a/Engine/Source/Editor/PropertyEditor/Private/IDetailsViewPrivate.h
+++ b/Engine/Source/Editor/PropertyEditor/Private/IDetailsViewPrivate.h
@@ -142,4 +142,7 @@ public:
 
 	/** Column width accessibility */
 	virtual FDetailColumnSizeData& GetColumnSizeData() = 0;
+
+	/** Does this details view allow favoriting? */
+	virtual bool IsFavoritingEnabled() const = 0;
 };
diff --git a/Engine/Source/Editor/PropertyEditor/Private/ItemPropertyNode.cpp b/Engine/Source/Editor/PropertyEditor/Private/ItemPropertyNode.cpp
--- a/Engine/Source/Editor/PropertyEditor/Private/ItemPropertyNode.cpp
+++ b/Engine/Source/Editor/PropertyEditor/Private/ItemPropertyNode.cpp
@@ -1,12 +1,13 @@
 // Copyright Epic Games, Inc. All Rights Reserved.
 
-
 #include "ItemPropertyNode.h"
-#include "Misc/ConfigCacheIni.h"
-#include "Classes/EditorStyleSettings.h"
+#include "DetailsViewConfig.h"
+#include "Editor.h"
+#include "EditorMetadataOverrides.h"
 #include "ObjectPropertyNode.h"
 #include "PropertyEditorHelpers.h"
 #include "PropertyPathHelpers.h"
+#include "Classes/EditorStyleSettings.h"
 
 #define LOCTEXT_NAMESPACE "ItemPropertyNode"
 
@@ -392,36 +393,26 @@ void FItemPropertyNode::InitChildNodes()
 
 void FItemPropertyNode::SetFavorite(bool FavoriteValue)
 {
-	const FObjectPropertyNode* CurrentObjectNode = FindObjectItemParent();
-	if (CurrentObjectNode == nullptr || CurrentObjectNode->GetNumObjects() <= 0)
-		return;
-	const UClass *ObjectClass = CurrentObjectNode->GetObjectBaseClass();
-	if (ObjectClass == nullptr)
-		return;
-	FString FullPropertyPath = ObjectClass->GetName() + TEXT(":") + PropertyPath;
-	if (FavoriteValue)
-	{
-		GConfig->SetBool(TEXT("DetailPropertyFavorites"), *FullPropertyPath, FavoriteValue, GEditorPerProjectIni);
-	}
-	else
+	UEditorMetadataOverrides* MetadataOverrides = GEditor->GetEditorSubsystem<UEditorMetadataOverrides>();
+	if (MetadataOverrides != nullptr)
 	{
-		GConfig->RemoveKey(TEXT("DetailPropertyFavorites"), *FullPropertyPath, GEditorPerProjectIni);
+		MetadataOverrides->SetBoolMetadata(GetProperty(), "IsFavorite", FavoriteValue);
 	}
 }
 
 bool FItemPropertyNode::IsFavorite() const
 {
-	const FObjectPropertyNode* CurrentObjectNode = FindObjectItemParent();
-	if (CurrentObjectNode == nullptr ||CurrentObjectNode->GetNumObjects() <= 0)
-		return false;
-	const UClass *ObjectClass = CurrentObjectNode->GetObjectBaseClass();
-	if (ObjectClass == nullptr)
-		return false;
-	FString FullPropertyPath = ObjectClass->GetName() + TEXT(":") + PropertyPath;
-	bool FavoritesPropertyValue = false;
-	if (!GConfig->GetBool(TEXT("DetailPropertyFavorites"), *FullPropertyPath, FavoritesPropertyValue, GEditorPerProjectIni))
-		return false;
-	return FavoritesPropertyValue;
+	UEditorMetadataOverrides* MetadataOverrides = GEditor->GetEditorSubsystem<UEditorMetadataOverrides>();
+	if (MetadataOverrides != nullptr)
+	{
+		bool IsFavorite = false;
+		if (MetadataOverrides->GetBoolMetadata(GetProperty(), "IsFavorite", IsFavorite))
+		{
+			return IsFavorite;
+		}
+	}
+
+	return false;
 }
 
 /**
diff --git a/Engine/Source/Editor/PropertyEditor/Private/SDetailsViewBase.cpp b/Engine/Source/Editor/PropertyEditor/Private/SDetailsViewBase.cpp
--- a/Engine/Source/Editor/PropertyEditor/Private/SDetailsViewBase.cpp
+++ b/Engine/Source/Editor/PropertyEditor/Private/SDetailsViewBase.cpp
@@ -387,7 +387,7 @@ void SDetailsViewBase::UpdateSinglePropertyMap(TSharedPtr<FComplexPropertyNode>
 	TSharedPtr<FComplexPropertyNode> RootPropertyNode = InRootPropertyNode;
 	check(RootPropertyNode.IsValid());
 
-	const bool bEnableFavoriteSystem = false; //TODO: IsEngineExitRequested() ? false : DetailsViewArgs.bAllowFavoriteSystem;
+	const bool bEnableFavoriteSystem = IsEngineExitRequested() ? false : DetailsViewArgs.bAllowFavoriteSystem;
 
 	DetailLayoutHelpers::FUpdatePropertyMapArgs Args;
 
@@ -405,15 +405,33 @@ void SDetailsViewBase::UpdateSinglePropertyMap(TSharedPtr<FComplexPropertyNode>
 	// inhibit our ability to find a single property's tree node. This is problematic for the diff and merge tools, that need
 	// to display and highlight each changed property for the user. We could whitelist 'good' customizations here if 
 	// we can make them work with the diff/merge tools.
-	if(!bDisableCustomDetailLayouts)
+	if (!bDisableCustomDetailLayouts)
 	{
 		DetailLayoutHelpers::QueryCustomDetailLayout(LayoutData, InstancedClassToDetailLayoutMap, GenericLayoutDelegate);
 	}
 
+	if (bEnableFavoriteSystem && InRootPropertyNode->GetInstancesNum() > 0)
+	{
+		static const FName FavoritesCategoryName("Favorites");
+		FDetailCategoryImpl& FavoritesCategory = LayoutData.DetailLayout->DefaultCategory(FavoritesCategoryName);
+		FavoritesCategory.SetSortOrder(0);
+		FavoritesCategory.SetCategoryAsSpecialFavorite();
+
+		if (FavoritesCategory.GetNumCustomizations() == 0)
+		{
+			FavoritesCategory.AddCustomRow(NSLOCTEXT("DetailLayoutHelpers", "Favorites", "Favorites"))
+				.WholeRowContent()
+				.VAlign(VAlign_Center)
+				.HAlign(HAlign_Center)
+				[
+					SNew(STextBlock).Text(NSLOCTEXT("DetailLayoutHelpers", "AddToFavoritesDescription", "Right-click on a property to add it to your Favorites."))
+				];
+		}
+	}
+
 	LayoutData.DetailLayout->GenerateDetailLayout();
 }
 
-
 void SDetailsViewBase::OnColorPickerWindowClosed(const TSharedRef<SWindow>& Window)
 {
 	const TSharedPtr< FPropertyNode > PinnedColorPropertyNode = ColorPropertyNode.Pin();
@@ -433,7 +451,6 @@ void SDetailsViewBase::OnColorPickerWindowClosed(const TSharedRef<SWindow>& Wind
 	ColorPropertyNode.Reset();
 }
 
-
 void SDetailsViewBase::SetIsPropertyVisibleDelegate(FIsPropertyVisible InIsPropertyVisible)
 {
 	IsPropertyVisibleDelegate = InIsPropertyVisible;
diff --git a/Engine/Source/Editor/PropertyEditor/Private/SDetailsViewBase.h b/Engine/Source/Editor/PropertyEditor/Private/SDetailsViewBase.h
--- a/Engine/Source/Editor/PropertyEditor/Private/SDetailsViewBase.h
+++ b/Engine/Source/Editor/PropertyEditor/Private/SDetailsViewBase.h
@@ -168,7 +168,7 @@ public:
 	virtual void SaveExpandedItems( TSharedRef<FPropertyNode> StartNode ) override;
 	virtual void RestoreExpandedItems(TSharedRef<FPropertyNode> StartNode) override;
 	virtual FDetailColumnSizeData& GetColumnSizeData() override { return ColumnSizeData; }
-
+	virtual bool IsFavoritingEnabled() const override { return DetailsViewArgs.bAllowFavoriteSystem; }
 	virtual bool IsConnected() const = 0;
 	virtual FRootPropertyNodeList& GetRootNodes() = 0;
 
diff --git a/Engine/Source/Editor/PropertyEditor/PropertyEditor.Build.cs b/Engine/Source/Editor/PropertyEditor/PropertyEditor.Build.cs
--- a/Engine/Source/Editor/PropertyEditor/PropertyEditor.Build.cs
+++ b/Engine/Source/Editor/PropertyEditor/PropertyEditor.Build.cs
@@ -65,6 +65,7 @@ public class PropertyEditor : ModuleRules
                 "SceneOutliner",
 				"DesktopPlatform",
 				"PropertyPath",
+				"EditorConfig",
             }
         );
 
diff --git a/Engine/Source/Editor/UnrealEd/Classes/Settings/EditorExperimentalSettings.h b/Engine/Source/Editor/UnrealEd/Classes/Settings/EditorExperimentalSettings.h
--- a/Engine/Source/Editor/UnrealEd/Classes/Settings/EditorExperimentalSettings.h
+++ b/Engine/Source/Editor/UnrealEd/Classes/Settings/EditorExperimentalSettings.h
@@ -52,12 +52,6 @@ public:
 	UPROPERTY(EditAnywhere, config, Category = Tools, meta = (DisplayName = "Translation Picker"))
 	bool bEnableTranslationPicker;
 
-	/** When enabled, all details panels will be able to have properties marked as favorite that show in a top most category.  
-	 * NOTE: Some customizations are not supported yet
-	 */
-	UPROPERTY(EditAnywhere, config, Category = Tools, meta = (DisplayName = "Enable Details Panel Favorites"))
-	bool bEnableFavoriteSystem;
-
 	/** Specify which console-specific nomenclature to use for gamepad label text */
 	UPROPERTY(EditAnywhere, config, Category=UserInterface, meta=(DisplayName="Console for Gamepad Labels"))
 	TEnumAsByte<EConsoleForGamepadLabels::Type> ConsoleForGamepadLabels;
diff --git a/Engine/Source/Editor/UnrealEd/Private/EditorEngine.cpp b/Engine/Source/Editor/UnrealEd/Private/EditorEngine.cpp
--- a/Engine/Source/Editor/UnrealEd/Private/EditorEngine.cpp
+++ b/Engine/Source/Editor/UnrealEd/Private/EditorEngine.cpp
@@ -1021,7 +1021,8 @@ void UEditorEngine::Init(IEngineLoop* InEngineLoop)
 			TEXT("TraceInsights"),
 			TEXT("StaticMeshEditor"),
 			TEXT("EditorFramework"),
-			TEXT("WorldPartitionEditor")
+			TEXT("WorldPartitionEditor"),
+			TEXT("EditorConfig")
 		};
 
 		FScopedSlowTask ModuleSlowTask(UE_ARRAY_COUNT(ModuleNames));
