Fixed CIS errors in previous submits.

diff --git a/Engine/Source/Editor/EditorConfig/Private/JsonConfig.cpp b/Engine/Source/Editor/EditorConfig/Private/JsonConfig.cpp
--- a/Engine/Source/Editor/EditorConfig/Private/JsonConfig.cpp
+++ b/Engine/Source/Editor/EditorConfig/Private/JsonConfig.cpp
@@ -2,9 +2,12 @@
 
 #include "JsonConfig.h"
 #include "Misc/FileHelper.h"
+#include "Misc/StringBuilder.h"
 #include "Serialization/JsonReader.h"
 #include "Serialization/JsonSerializer.h"
 
+DEFINE_LOG_CATEGORY_STATIC(LogJsonConfig, Log, All);
+
 namespace UE
 {
 	FJsonPath::FJsonPath()
@@ -47,7 +50,7 @@ namespace UE
 			if (Part.IsEmpty())
 			{
 				PathParts.Empty();
-				UE_LOG(LogEditorConfig, Warning, TEXT("Path part was empty in JSON path \"%s\""), *InPath);
+				UE_LOG(LogJsonConfig, Warning, TEXT("Path part was empty in JSON path \"%s\""), *InPath);
 				return;
 			}
 
@@ -71,14 +74,14 @@ namespace UE
 					else
 					{
 						PathParts.Empty();
-						UE_LOG(LogEditorConfig, Warning, TEXT("Path part \"%s\" looked to be index into array but was malformed in JSON path \"%s\""), *Part, *InPath);
+						UE_LOG(LogJsonConfig, Warning, TEXT("Path part \"%s\" looked to be index into array but was malformed in JSON path \"%s\""), *Part, *InPath);
 						return;
 					}
 				}
 				else
 				{
 					PathParts.Empty();
-					UE_LOG(LogEditorConfig, Warning, TEXT("Path part \"%s\" looked to be index into array but was malformed in JSON path \"%s\""), *Part, *InPath);
+					UE_LOG(LogJsonConfig, Warning, TEXT("Path part \"%s\" looked to be index into array but was malformed in JSON path \"%s\""), *Part, *InPath);
 					return;
 				}
 			}
@@ -400,7 +403,7 @@ namespace UE
 		if (OverrideSet != nullptr && 
 			(OverrideAdd != nullptr || OverrideRemove != nullptr))
 		{
-			UE_LOG(LogEditorConfig, Warning, TEXT("JSON container is malformed and contains both set and add/remove values."));
+			UE_LOG(LogJsonConfig, Warning, TEXT("JSON container is malformed and contains both set and add/remove values."));
 		}
 
 		return OverrideSet != nullptr || OverrideAdd != nullptr || OverrideRemove != nullptr;
diff --git a/Engine/Source/Editor/EditorConfig/Public/EditorConfig.h b/Engine/Source/Editor/EditorConfig/Public/EditorConfig.h
--- a/Engine/Source/Editor/EditorConfig/Public/EditorConfig.h
+++ b/Engine/Source/Editor/EditorConfig/Public/EditorConfig.h
@@ -7,8 +7,6 @@
 #include "Dom/JsonObject.h"
 #include "Templates/Models.h"
 
-DEFINE_LOG_CATEGORY_STATIC(LogEditorConfig, Log, All);
-
 class EDITORCONFIG_API FEditorConfig
 {
 public:
@@ -79,7 +77,8 @@ bool FEditorConfig::TryGetStruct(FStringView Key, T& OutValue) const
 	}
 
 	TSharedPtr<FJsonObject> StructData;
-	if (!JsonConfig->TryGetJsonObject(UE::FJsonPath(Key), StructData))
+	UE::FJsonPath Path(Key);
+	if (!JsonConfig->TryGetJsonObject(Path, StructData))
 	{
 		return false;
 	}
@@ -106,7 +105,8 @@ bool FEditorConfig::TryGetUObject(FStringView Key, T& OutValue) const
 	}
 
 	TSharedPtr<FJsonObject> UObjectData;
-	if (!JsonConfig->TryGetJsonObject(UE::FJsonPath(Key), UObjectData))
+	UE::FJsonPath Path(Key);
+	if (!JsonConfig->TryGetJsonObject(Path, UObjectData))
 	{
 		return false;
 	}
