Fixed CIS errors: static_assert with no message is a C++17 extension.

diff --git a/Engine/Source/Editor/EditorConfig/Public/EditorConfig.h b/Engine/Source/Editor/EditorConfig/Public/EditorConfig.h
--- a/Engine/Source/Editor/EditorConfig/Public/EditorConfig.h
+++ b/Engine/Source/Editor/EditorConfig/Public/EditorConfig.h
@@ -97,7 +97,7 @@ bool FEditorConfig::TryGetStruct(FStringView Key, T& OutValue) const
 template <typename T>
 bool FEditorConfig::TryGetUObject(FStringView Key, T& OutValue) const
 {
-	static_assert(TIsDerivedFrom<T, UObject>::Value);
+	static_assert(TIsDerivedFrom<T, UObject>::Value, "Type is not derived from UObject.");
 
 	if (!IsValid())
 	{
@@ -141,7 +141,7 @@ bool FEditorConfig::TryGetRootStruct(T& OutValue) const
 template <typename T>
 bool FEditorConfig::TryGetRootUObject(T& OutValue) const
 {
-	static_assert(TIsDerivedFrom<T, UObject>::Value);
+	static_assert(TIsDerivedFrom<T, UObject>::Value, "Type is not derived from UObject.");
 
 	if (!IsValid())
 	{
@@ -173,7 +173,7 @@ void FEditorConfig::SetStruct(FStringView Key, const T& InValue)
 template <typename T>
 void FEditorConfig::SetUObject(FStringView Key, const T& InValue)
 {
-	static_assert(TIsDerivedFrom<T, UObject>::Value);
+	static_assert(TIsDerivedFrom<T, UObject>::Value, "Type is not derived from UObject.");
 
 	if (!IsValid())
 	{
@@ -203,7 +203,7 @@ void FEditorConfig::SetRootStruct(const T& InValue)
 template <typename T>
 void FEditorConfig::SetRootUObject(const T& InValue)
 {
-	static_assert(TIsDerivedFrom<T, UObject>::Value);
+	static_assert(TIsDerivedFrom<T, UObject>::Value, "Type is not derived from UObject.");
 
 	if (!IsValid())
 	{
