Removed a dangling PRAGMA_ENABLE_OPTIMIZATION. Added copyright notice.

diff --git a/Engine/Source/Editor/EditorConfig/Private/EditorMetadataOverrides.cpp b/Engine/Source/Editor/EditorConfig/Private/EditorMetadataOverrides.cpp
--- a/Engine/Source/Editor/EditorConfig/Private/EditorMetadataOverrides.cpp
+++ b/Engine/Source/Editor/EditorConfig/Private/EditorMetadataOverrides.cpp
@@ -348,5 +348,3 @@ void UEditorMetadataOverrides::RemoveMetadata(const FField* Field, FName Key)
 
 	bDirty = true;
 }
-
-PRAGMA_ENABLE_OPTIMIZATION
\ No newline at end of file
diff --git a/Engine/Source/Editor/EditorConfig/Private/Tests/EditorConfigTests.h b/Engine/Source/Editor/EditorConfig/Private/Tests/EditorConfigTests.h
--- a/Engine/Source/Editor/EditorConfig/Private/Tests/EditorConfigTests.h
+++ b/Engine/Source/Editor/EditorConfig/Private/Tests/EditorConfigTests.h
@@ -1,3 +1,5 @@
+// Copyright Epic Games, Inc. All Rights Reserved.
+
 #pragma once
 
 #include "CoreMinimal.h"
