Fixed non-unity build errors due to missing semicolon after IMPLEMENT_MODULE.

Removed unnecessary includes and overrides in EditorConfigModule.

diff --git a/Engine/Source/Editor/EditorConfig/Private/EditorConfigModule.cpp b/Engine/Source/Editor/EditorConfig/Private/EditorConfigModule.cpp
--- a/Engine/Source/Editor/EditorConfig/Private/EditorConfigModule.cpp
+++ b/Engine/Source/Editor/EditorConfig/Private/EditorConfigModule.cpp
@@ -1,17 +1,6 @@
 // Copyright Epic Games, Inc. All Rights Reserved.
 
 #include "EditorConfigModule.h"
-#include "EditorConfig.h"
-#include "EditorConfigSubsystem.h"
-#include "EditorMetadataOverrides.h"
-#include "Misc/Paths.h"
+#include "Modules/ModuleManager.h"
 
-void FEditorConfigModule::StartupModule()
-{
-}
-
-void FEditorConfigModule::ShutdownModule()
-{
-}
-
-IMPLEMENT_MODULE(FEditorConfigModule, EditorConfig)
\ No newline at end of file
+IMPLEMENT_MODULE(FEditorConfigModule, EditorConfig);
diff --git a/Engine/Source/Editor/EditorConfig/Public/EditorConfigModule.h b/Engine/Source/Editor/EditorConfig/Public/EditorConfigModule.h
--- a/Engine/Source/Editor/EditorConfig/Public/EditorConfigModule.h
+++ b/Engine/Source/Editor/EditorConfig/Public/EditorConfigModule.h
@@ -6,7 +6,4 @@
 
 class FEditorConfigModule : public IModuleInterface
 {
-public:
-	virtual void StartupModule() override;
-	virtual void ShutdownModule() override;
 };
