Fixed incorrect operator usage in FJsonConfig picked up by static analysis.

diff --git a/Engine/Source/Editor/EditorConfig/Private/JsonConfig.cpp b/Engine/Source/Editor/EditorConfig/Private/JsonConfig.cpp
--- a/Engine/Source/Editor/EditorConfig/Private/JsonConfig.cpp
+++ b/Engine/Source/Editor/EditorConfig/Private/JsonConfig.cpp
@@ -1007,7 +1007,7 @@ namespace UE
 				return true;
 
 			case EJson::Array:
-				if (CurrentValue->Type != EJson::Object || CurrentValue->Type != EJson::Array)
+				if (CurrentValue->Type != EJson::Object && CurrentValue->Type != EJson::Array)
 				{
 					return false;
 				}
