Added checks to satisfy PVS.

diff --git a/Engine/Source/Editor/EditorConfig/Private/JsonConfig.cpp b/Engine/Source/Editor/EditorConfig/Private/JsonConfig.cpp
--- a/Engine/Source/Editor/EditorConfig/Private/JsonConfig.cpp
+++ b/Engine/Source/Editor/EditorConfig/Private/JsonConfig.cpp
@@ -323,6 +323,8 @@ namespace UE
 		const TArray<TSharedPtr<FJsonValue>>* OverrideSet = nullptr;
 		if (OverrideObject->TryGetArrayField(TEXT("="), OverrideSet))
 		{
+			check(OverrideSet != nullptr);
+
 			DestArray.Reset();
 			DestArray.Reserve(OverrideSet->Num());
 			for (const TSharedPtr<FJsonValue>& SetValue : *OverrideSet)
@@ -334,6 +336,8 @@ namespace UE
 		const TArray<TSharedPtr<FJsonValue>>* OverrideAdd = nullptr;
 		if (OverrideObject->TryGetArrayField(TEXT("+"), OverrideAdd))
 		{
+			check(OverrideAdd != nullptr);
+
 			DestArray.Reserve(DestArray.Num() + OverrideAdd->Num());
 			for (const TSharedPtr<FJsonValue>& AddValue : *OverrideAdd)
 			{
@@ -374,6 +378,8 @@ namespace UE
 		const TArray<TSharedPtr<FJsonValue>>* OverrideRemove = nullptr;
 		if (OverrideObject->TryGetArrayField(TEXT("-"), OverrideRemove))
 		{
+			check(OverrideRemove != nullptr);
+
 			for (const TSharedPtr<FJsonValue>& RemoveValue : *OverrideRemove)
 			{
 				DestArray.RemoveAll(
@@ -416,9 +422,11 @@ namespace UE
 			TSharedPtr<FJsonValue> DestValue = Dest->TryGetField(Override.Key);
 			if (!DestValue.IsValid())
 			{
-				const TSharedPtr<FJsonObject>* OverrideObject;
+				const TSharedPtr<FJsonObject>* OverrideObject = nullptr;
 				if (Override.Value->TryGetObject(OverrideObject))
 				{
+					check(OverrideObject);
+
 					TArray<TSharedPtr<FJsonValue>> DestArray;
 					if (ApplyDeltaOperationsToArray(*OverrideObject, DestArray))
 					{
@@ -483,13 +491,15 @@ namespace UE
 					}
 					case EJson::Object:
 					{
-						const TSharedPtr<FJsonObject>* OverrideObject;
+						const TSharedPtr<FJsonObject>* OverrideObject = nullptr;
 						if (Override.Value->TryGetObject(OverrideObject))
 						{
+							check(OverrideObject);
 							TArray<TSharedPtr<FJsonValue>>* DestArray = nullptr;
 							TSharedPtr<FJsonObject>* DestObject = nullptr;
 							if (DestValue->TryGetObject(DestObject))
 							{
+								check(DestObject);
 								ApplyOverridesToObject(*OverrideObject, *DestObject);
 							}
 							else if (DestValue->TryGetArray(DestArray))
