Changed conditions in JsonConfig.cpp to satisfy PVS.

diff --git a/Engine/Source/Editor/EditorConfig/Private/JsonConfig.cpp b/Engine/Source/Editor/EditorConfig/Private/JsonConfig.cpp
--- a/Engine/Source/Editor/EditorConfig/Private/JsonConfig.cpp
+++ b/Engine/Source/Editor/EditorConfig/Private/JsonConfig.cpp
@@ -321,10 +321,8 @@ namespace UE
 	static bool ApplyDeltaOperationsToArray(const TSharedPtr<FJsonObject>& OverrideObject, TArray<TSharedPtr<FJsonValue>>& DestArray)
 	{
 		const TArray<TSharedPtr<FJsonValue>>* OverrideSet = nullptr;
-		if (OverrideObject->TryGetArrayField(TEXT("="), OverrideSet))
+		if (OverrideObject->TryGetArrayField(TEXT("="), OverrideSet) && OverrideSet) // the latter part of this condition is just to satisfy static analysis
 		{
-			check(OverrideSet != nullptr);
-
 			DestArray.Reset();
 			DestArray.Reserve(OverrideSet->Num());
 			for (const TSharedPtr<FJsonValue>& SetValue : *OverrideSet)
@@ -334,10 +332,8 @@ namespace UE
 		}
 
 		const TArray<TSharedPtr<FJsonValue>>* OverrideAdd = nullptr;
-		if (OverrideObject->TryGetArrayField(TEXT("+"), OverrideAdd))
+		if (OverrideObject->TryGetArrayField(TEXT("+"), OverrideAdd) && OverrideAdd)
 		{
-			check(OverrideAdd != nullptr);
-
 			DestArray.Reserve(DestArray.Num() + OverrideAdd->Num());
 			for (const TSharedPtr<FJsonValue>& AddValue : *OverrideAdd)
 			{
@@ -376,10 +372,8 @@ namespace UE
 		}
 
 		const TArray<TSharedPtr<FJsonValue>>* OverrideRemove = nullptr;
-		if (OverrideObject->TryGetArrayField(TEXT("-"), OverrideRemove))
+		if (OverrideObject->TryGetArrayField(TEXT("-"), OverrideRemove) && OverrideRemove)
 		{
-			check(OverrideRemove != nullptr);
-
 			for (const TSharedPtr<FJsonValue>& RemoveValue : *OverrideRemove)
 			{
 				DestArray.RemoveAll(
