Actually fixed PVS warnings in JsonConfig.cpp.

diff --git a/Engine/Source/Editor/EditorConfig/Private/JsonConfig.cpp b/Engine/Source/Editor/EditorConfig/Private/JsonConfig.cpp
--- a/Engine/Source/Editor/EditorConfig/Private/JsonConfig.cpp
+++ b/Engine/Source/Editor/EditorConfig/Private/JsonConfig.cpp
@@ -323,9 +323,11 @@ namespace UE
 		const TArray<TSharedPtr<FJsonValue>>* OverrideSet = nullptr;
 		if (OverrideObject->TryGetArrayField(TEXT("="), OverrideSet) && OverrideSet) // the latter part of this condition is just to satisfy static analysis
 		{
+			const TArray<TSharedPtr<FJsonValue>>& OverrideSetRef = *OverrideSet;
+
 			DestArray.Reset();
-			DestArray.Reserve(OverrideSet->Num());
-			for (const TSharedPtr<FJsonValue>& SetValue : *OverrideSet)
+			DestArray.Reserve(OverrideSetRef.Num());
+			for (const TSharedPtr<FJsonValue>& SetValue : OverrideSetRef)
 			{
 				DestArray.Add(FJsonValue::Duplicate(SetValue));
 			}
@@ -334,8 +336,10 @@ namespace UE
 		const TArray<TSharedPtr<FJsonValue>>* OverrideAdd = nullptr;
 		if (OverrideObject->TryGetArrayField(TEXT("+"), OverrideAdd) && OverrideAdd)
 		{
-			DestArray.Reserve(DestArray.Num() + OverrideAdd->Num());
-			for (const TSharedPtr<FJsonValue>& AddValue : *OverrideAdd)
+			const TArray<TSharedPtr<FJsonValue>>& OverrideAddRef = *OverrideAdd;
+
+			DestArray.Reserve(DestArray.Num() + OverrideAddRef.Num());
+			for (const TSharedPtr<FJsonValue>& AddValue : OverrideAddRef)
 			{
 				// check if this is a map with $key, $value fields
 				if (AddValue->Type == EJson::Object)
@@ -374,7 +378,8 @@ namespace UE
 		const TArray<TSharedPtr<FJsonValue>>* OverrideRemove = nullptr;
 		if (OverrideObject->TryGetArrayField(TEXT("-"), OverrideRemove) && OverrideRemove)
 		{
-			for (const TSharedPtr<FJsonValue>& RemoveValue : *OverrideRemove)
+			const TArray<TSharedPtr<FJsonValue>>& OverrideRemoveRef = *OverrideRemove;
+			for (const TSharedPtr<FJsonValue>& RemoveValue : OverrideRemoveRef)
 			{
 				DestArray.RemoveAll(
 					[RemoveValue](const TSharedPtr<FJsonValue>& Value)
@@ -417,10 +422,8 @@ namespace UE
 			if (!DestValue.IsValid())
 			{
 				const TSharedPtr<FJsonObject>* OverrideObject = nullptr;
-				if (Override.Value->TryGetObject(OverrideObject))
+				if (Override.Value->TryGetObject(OverrideObject) && OverrideObject)
 				{
-					check(OverrideObject);
-
 					TArray<TSharedPtr<FJsonValue>> DestArray;
 					if (ApplyDeltaOperationsToArray(*OverrideObject, DestArray))
 					{
@@ -491,12 +494,11 @@ namespace UE
 							check(OverrideObject);
 							TArray<TSharedPtr<FJsonValue>>* DestArray = nullptr;
 							TSharedPtr<FJsonObject>* DestObject = nullptr;
-							if (DestValue->TryGetObject(DestObject))
+							if (DestValue->TryGetObject(DestObject) && DestObject)
 							{
-								check(DestObject);
 								ApplyOverridesToObject(*OverrideObject, *DestObject);
 							}
-							else if (DestValue->TryGetArray(DestArray))
+							else if (DestValue->TryGetArray(DestArray) && DestArray)
 							{
 								ApplyDeltaOperationsToArray(*OverrideObject, *DestArray);
 							}
@@ -1235,6 +1237,7 @@ namespace UE
 			return true;
 		}
 	
+		check(LastValue);
 		return SetValueHelper(*LastValue, NewValue, ParentValue);
 	}
 
