Added missing includes causing non-Unity compile failures.

diff --git a/Engine/Source/Editor/EditorConfig/Private/EditorConfig.cpp b/Engine/Source/Editor/EditorConfig/Private/EditorConfig.cpp
--- a/Engine/Source/Editor/EditorConfig/Private/EditorConfig.cpp
+++ b/Engine/Source/Editor/EditorConfig/Private/EditorConfig.cpp
@@ -1,6 +1,8 @@
 // Copyright Epic Games, Inc. All Rights Reserved.
 
 #include "EditorConfig.h"
+#include "UObject/EnumProperty.h"
+#include "UObject/TextProperty.h"
 #include "UObject/UnrealType.h"
 
 FEditorConfig::FEditorConfig()
