Adding init values to FEditorConfigTestKeyStruct and FEditorConfigTestStruct
#pf 60622a2125b98e00011377ed

diff --git a/Engine/Source/Editor/EditorConfig/Private/Tests/EditorConfigTests.h b/Engine/Source/Editor/EditorConfig/Private/Tests/EditorConfigTests.h
--- a/Engine/Source/Editor/EditorConfig/Private/Tests/EditorConfigTests.h
+++ b/Engine/Source/Editor/EditorConfig/Private/Tests/EditorConfigTests.h
@@ -13,16 +13,16 @@ struct FEditorConfigTestStruct
 	GENERATED_BODY()
 
 	UPROPERTY()
-	bool BoolProperty;
+	bool BoolProperty = false;
 
 	UPROPERTY()
-	int32 IntProperty;
+	int32 IntProperty = 0;
 
 	UPROPERTY()
 	FString StringProperty;
 
 	UPROPERTY()
-	float FloatProperty;
+	float FloatProperty = 0.0f;
 
 	UPROPERTY()
 	TArray<FString> ArrayProperty;
@@ -37,7 +37,7 @@ struct FEditorConfigTestKeyStruct
 	FString Name;
 	
 	UPROPERTY()
-	double Number;
+	double Number = 0.0;
 
 	bool operator==(const FEditorConfigTestKeyStruct& Other) const;
 	bool operator!=(const FEditorConfigTestKeyStruct& Other) const;
