Fixed potential nullptr deref caught by static analysis.

diff --git a/Engine/Source/Editor/EditorConfig/Private/EditorMetadataOverrides.cpp b/Engine/Source/Editor/EditorConfig/Private/EditorMetadataOverrides.cpp
--- a/Engine/Source/Editor/EditorConfig/Private/EditorMetadataOverrides.cpp
+++ b/Engine/Source/Editor/EditorConfig/Private/EditorMetadataOverrides.cpp
@@ -550,7 +550,7 @@ bool UEditorMetadataOverrides::GetStringMetadata(const UStruct* Struct, FName Ke
 	}
 
 	const FString* MetaValue = StructMetadata->Strings.Find(Key);
-	if (MetaValue != nullptr)
+	if (MetaValue == nullptr)
 	{
 		if (!Struct->HasMetaData(Key))
 		{
