UEditorConfigSubsystem now has helpers to directly load into UObjects based on the EditorConfig= class metadata, as well as filtering of properties tagged with the EditorConfig property metadata.

diff --git a/Engine/Source/Editor/EditorConfig/Private/EditorConfig.cpp b/Engine/Source/Editor/EditorConfig/Private/EditorConfig.cpp
--- a/Engine/Source/Editor/EditorConfig/Private/EditorConfig.cpp
+++ b/Engine/Source/Editor/EditorConfig/Private/EditorConfig.cpp
@@ -81,54 +81,86 @@ bool FEditorConfig::HasOverride(FStringView Key) const
 	return JsonConfig->HasOverride(UE::FJsonPath(Key));
 }
 
-void FEditorConfig::ReadStruct(TSharedPtr<FJsonObject> JsonObject, const UStruct* Struct, void* Instance, UObject* Owner)
+void FEditorConfig::SetRootUObject(const UClass* Class, const UObject* Instance, EPropertyFilter Filter)
+{
+	if (!IsValid())
+	{
+		return;
+	}
+	
+	TSharedPtr<FJsonObject> JsonObject = WriteUObject(Class, Instance, Filter);
+	JsonConfig->SetRootObject(JsonObject);
+		
+	SetDirty();
+}
+
+void FEditorConfig::SetRootStruct(const UStruct* Struct, const void* Instance, EPropertyFilter Filter)
+{
+	if (!IsValid())
+	{
+		return;
+	}
+	
+	TSharedPtr<FJsonObject> JsonObject = WriteStruct(Struct, Instance, Filter);
+	JsonConfig->SetRootObject(JsonObject);
+		
+	SetDirty();
+}
+
+void FEditorConfig::ReadStruct(TSharedPtr<FJsonObject> JsonObject, const UStruct* Struct, void* Instance, UObject* Owner, EPropertyFilter Filter)
 {
 	FString TypeName;
 	JsonObject->TryGetStringField(TEXT("$type"), TypeName);
 
-	if (!ensureAlways(Struct->GetName() != TypeName))
+	if (!TypeName.IsEmpty() && !ensureAlwaysMsgf(Struct->GetName().Equals(TypeName), TEXT("Type name mismatch in FEditorConfig::ReadUObject. Expected: %s, Actual: %s"), *Struct->GetName(), *TypeName))
 	{
 		return;
 	}
 
 	for (TFieldIterator<FProperty> It(Struct); It; ++It)
 	{
-		FProperty* Property = *It;
-		if (Property != nullptr)
+		const FProperty* Property = *It;
+
+		if (Filter == EPropertyFilter::MetadataOnly && !Property->HasMetaData("EditorConfig"))
 		{
-			void* DataPtr = Property->ContainerPtrToValuePtr<void>(Instance);
+			continue;
+		}
+
+		void* DataPtr = Property->ContainerPtrToValuePtr<void>(Instance);
 			
-			TSharedPtr<FJsonValue> Value = JsonObject->TryGetField(Property->GetName());
-			if (Value.IsValid())
-			{
-				ReadValue(Value, Property, DataPtr, Owner);
-			}
+		TSharedPtr<FJsonValue> Value = JsonObject->TryGetField(Property->GetName());
+		if (Value.IsValid())
+		{
+			ReadValue(Value, Property, DataPtr, Owner);
 		}
 	}
 }
 
-void FEditorConfig::ReadUObject(TSharedPtr<FJsonObject> JsonObject, const UClass* Class, UObject* Instance)
+void FEditorConfig::ReadUObject(TSharedPtr<FJsonObject> JsonObject, const UClass* Class, UObject* Instance, EPropertyFilter Filter)
 {
 	FString TypeName;
 	JsonObject->TryGetStringField(TEXT("$type"), TypeName);
 
-	if (!ensureAlways(Class->GetName() != TypeName))
+	if (!TypeName.IsEmpty() && !ensureAlwaysMsgf(Class->GetName().Equals(TypeName), TEXT("Type name mismatch in FEditorConfig::ReadUObject. Expected: %s, Actual: %s"), *Class->GetName(), *TypeName))
 	{
 		return;
 	}
 
 	for (TFieldIterator<FProperty> It(Class); It; ++It)
 	{
-		FProperty* Property = *It;
-		if (Property != nullptr)
+		const FProperty* Property = *It;
+		
+		if (Filter == EPropertyFilter::MetadataOnly && !Property->HasMetaData("EditorConfig"))
 		{
-			void* DataPtr = Property->ContainerPtrToValuePtr<void>(Instance);
+			continue;
+		}
 
-			TSharedPtr<FJsonValue> Value = JsonObject->TryGetField(Property->GetName());
-			if (Value.IsValid())
-			{
-				ReadValue(Value, Property, DataPtr, Instance);
-			}
+		void* DataPtr = Property->ContainerPtrToValuePtr<void>(Instance);
+
+		TSharedPtr<FJsonValue> Value = JsonObject->TryGetField(Property->GetName());
+		if (Value.IsValid())
+		{
+			ReadValue(Value, Property, DataPtr, Instance);
 		}
 	}
 }
@@ -243,7 +275,7 @@ void FEditorConfig::ReadValue(TSharedPtr<FJsonValue> JsonValue, const FProperty*
 		const TSharedPtr<FJsonObject>* ObjectJsonValue;
 		if (JsonValue->TryGetObject(ObjectJsonValue))
 		{
-			ReadStruct(*ObjectJsonValue, StructProperty->Struct, DataPtr, Owner);
+			ReadStruct(*ObjectJsonValue, StructProperty->Struct, DataPtr, Owner, EPropertyFilter::All);
 		}
 	}
 	else if (const FArrayProperty* ArrayProperty = CastField<FArrayProperty>(Property))
@@ -454,13 +486,12 @@ TSharedPtr<FJsonValue> FEditorConfig::WriteValue(const FProperty* Property, cons
 	else if (const FObjectPropertyBase* ObjectProperty = CastField<FObjectPropertyBase>(Property))
 	{
 		FString ObjectPath;
-		// TODO: No idea how we're meant to get object paths, might be this:
-		//ObjectProperty->ExportTextItem(ObjectPath, DataPtr, ??? /*DefaultValue*/, ??? /*Parent*/, 0, ??? /*ExportRootScope*/);
+		ObjectProperty->ExportTextItem(ObjectPath, DataPtr, nullptr, nullptr, PPF_None, nullptr);
 		ResultValue = MakeShared<FJsonValueString>(ObjectPath);
 	}
 	else if (const FStructProperty* StructProperty = CastField<FStructProperty>(Property))
 	{
-		ResultValue = MakeShared<FJsonValueObject>(WriteStruct(StructProperty->Struct, DataPtr));
+		ResultValue = MakeShared<FJsonValueObject>(WriteStruct(StructProperty->Struct, DataPtr, EPropertyFilter::All));
 	}
 	else if (const FArrayProperty* ArrayProperty = CastField<FArrayProperty>(Property))
 	{
@@ -526,7 +557,7 @@ TSharedPtr<FJsonValue> FEditorConfig::WriteValue(const FProperty* Property, cons
 			}
 
 			// maps can either be stored as $key, $value pairs or, if the keys can be stringified, as a JSON object
-			// check which we should use based on the first element
+			// check Filter we should use based on the first element
 			EJson KeyType = JsonKeysArray[0]->Type;
 			if (KeyType == EJson::Object)
 			{
@@ -591,13 +622,20 @@ bool FEditorConfig::IsDefault(const FProperty* Property, TSharedPtr<FJsonValue>
 	return false;
 }
 
-TSharedPtr<FJsonObject> FEditorConfig::WriteStruct(const UStruct* Struct, const void* Instance) 
+TSharedPtr<FJsonObject> FEditorConfig::WriteStruct(const UStruct* Struct, const void* Instance, EPropertyFilter Filter) 
 {
 	TSharedPtr<FJsonObject> JsonObject = MakeShared<FJsonObject>();
+	JsonObject->SetStringField(TEXT("$type"), Struct->GetName());
 
 	for (TFieldIterator<FProperty> It(Struct); It; ++It)
 	{
 		const FProperty* Property = *It;
+
+		if (Filter == EPropertyFilter::MetadataOnly && !Property->HasMetaData("EditorConfig"))
+		{
+			continue;
+		}
+
 		const void* ValuePtr = Property->ContainerPtrToValuePtr<void>(Instance); 
 
 		TSharedPtr<FJsonValue> PropertyValue = WriteValue(Property, ValuePtr);
@@ -611,15 +649,22 @@ TSharedPtr<FJsonObject> FEditorConfig::WriteStruct(const UStruct* Struct, const
 }
 
 /** 
-	* This exists because of sparse class data that can exist for UObjects only, and which is handled in ContainerPtrToValuePtr.
-	*/
-TSharedPtr<FJsonObject> FEditorConfig::WriteUObject(const UStruct* Struct, const UObject* Instance) 
+ * This exists because of sparse class data that can exist for UObjects only, which is handled in ContainerPtrToValuePtr.
+ */
+TSharedPtr<FJsonObject> FEditorConfig::WriteUObject(const UClass* Class, const UObject* Instance, EPropertyFilter Filter) 
 {
 	TSharedPtr<FJsonObject> JsonObject = MakeShared<FJsonObject>();
+	JsonObject->SetStringField(TEXT("$type"), Class->GetName());
 
-	for (TFieldIterator<FProperty> It(Struct); It; ++It)
+	for (TFieldIterator<FProperty> It(Class); It; ++It)
 	{
 		const FProperty* Property = *It;
+
+		if (Filter == EPropertyFilter::MetadataOnly && !Property->HasMetaData("EditorConfig"))
+		{
+			continue;
+		}
+
 		const void* ValuePtr = Property->ContainerPtrToValuePtr<void>(Instance); 
 
 		TSharedPtr<FJsonValue> PropertyValue = WriteValue(Property, ValuePtr);
@@ -631,3 +676,17 @@ TSharedPtr<FJsonObject> FEditorConfig::WriteUObject(const UStruct* Struct, const
 
 	return JsonObject;
 }
+
+void FEditorConfig::SetDirty()
+{
+	if (!Dirty)
+	{
+		Dirty = true;
+		EditorConfigDirtiedEvent.Broadcast(*this);
+	}
+}
+
+void FEditorConfig::OnSaved()
+{
+	Dirty = false;
+}
diff --git a/Engine/Source/Editor/EditorConfig/Private/EditorConfigBase.cpp b/Engine/Source/Editor/EditorConfig/Private/EditorConfigBase.cpp
new file mode 100644
--- /dev/null
+++ b/Engine/Source/Editor/EditorConfig/Private/EditorConfigBase.cpp
@@ -0,0 +1,17 @@
+// Copyright Epic Games, Inc. All Rights Reserved.
+
+#include "EditorConfigBase.h"
+#include "EditorConfigSubsystem.h"
+#include "Editor.h"
+
+bool UEditorConfigBase::LoadEditorConfig()
+{
+	UEditorConfigSubsystem* Subsystem = GEditor->GetEditorSubsystem<UEditorConfigSubsystem>();
+	return Subsystem->LoadConfigObject(this, this->GetClass());
+}
+
+bool UEditorConfigBase::SaveEditorConfig() const
+{
+	UEditorConfigSubsystem* Subsystem = GEditor->GetEditorSubsystem<UEditorConfigSubsystem>();
+	return Subsystem->SaveConfigObject(this, this->GetClass());
+}
\ No newline at end of file
diff --git a/Engine/Source/Editor/EditorConfig/Private/EditorConfigSubystem.cpp b/Engine/Source/Editor/EditorConfig/Private/EditorConfigSubystem.cpp
--- a/Engine/Source/Editor/EditorConfig/Private/EditorConfigSubystem.cpp
+++ b/Engine/Source/Editor/EditorConfig/Private/EditorConfigSubystem.cpp
@@ -14,59 +14,174 @@ UEditorConfigSubsystem::UEditorConfigSubsystem()
 
 void UEditorConfigSubsystem::Initialize(FSubsystemCollectionBase& Collection)
 {
-	AddSearchDirectory(FPaths::Combine(FPaths::Combine(FPlatformProcess::UserSettingsDir(), *FApp::GetEpicProductIdentifier(), TEXT("Editor")))); // AppData
 	AddSearchDirectory(FPaths::Combine(FPaths::EngineConfigDir(), TEXT("Editor"))); // Engine
 	AddSearchDirectory(FPaths::Combine(FPaths::ProjectConfigDir(), TEXT("Editor"))); // ProjectName
+	AddSearchDirectory(FPaths::Combine(FPlatformProcess::UserSettingsDir(), *FApp::GetEpicProductIdentifier(), TEXT("Editor"))); // AppData
 }
 
-TSharedPtr<FEditorConfig> UEditorConfigSubsystem::FindOrLoadConfig(FStringView ConfigName)
+void UEditorConfigSubsystem::Tick(float DeltaTime)
 {
+	SaveLock.WriteLock();
+	ON_SCOPE_EXIT { SaveLock.WriteUnlock(); };
+
+	for (FPendingSave& Save : PendingSaves)
+	{
+		Save.TimeSinceQueued += DeltaTime;
+
+		const float SaveDelaySeconds = 3.0f;
+		if (Save.TimeSinceQueued > SaveDelaySeconds)
+		{
+			const FString* FilePath = LoadedConfigs.FindKey(Save.Config);
+			check(FilePath != nullptr);
+			
+			Save.WasSuccess = Async(EAsyncExecution::Thread, 
+				[Config = Save.Config, File = *FilePath]()
+				{
+					return Config->SaveToFile(File);
+				}, 
+				[this, Config = Save.Config]()
+				{
+					OnSaveCompleted(Config);
+				});
+		}
+	}
+}
+
+TStatId UEditorConfigSubsystem::GetStatId() const
+{
+	RETURN_QUICK_DECLARE_CYCLE_STAT(UEditorConfigSubsystem, STATGROUP_Tickables);
+}
+
+bool UEditorConfigSubsystem::LoadConfigObject(UObject* Object, const UClass* Class, FEditorConfig::EPropertyFilter Filter)
+{
+	const FString& EditorConfigName = Class->GetMetaData("EditorConfig");
+	if (!ensureMsgf(!EditorConfigName.IsEmpty(), TEXT("UEditorConfigSubsystem::LoadConfigObject - EditorConfig name is not set on class %s."), *Class->GetName()))
+	{
+		return false;
+	}
+
+	TSharedRef<FEditorConfig> EditorConfig = FindOrLoadConfig(EditorConfigName);
+	return EditorConfig->TryGetRootUObject(*Object, Filter);
+}
+
+bool UEditorConfigSubsystem::SaveConfigObject(const UObject* Object, const UClass* Class, FEditorConfig::EPropertyFilter Filter)
+{
+	const FString& EditorConfigName = Class->GetMetaData("EditorConfig");
+	if (!ensureMsgf(EditorConfigName.IsEmpty(), TEXT("UEditorConfigSubsystem::SaveConfigObject - EditorConfig name is not set on class %s."), *Class->GetName()))
+	{
+		return false;
+	}
+
+	TSharedRef<FEditorConfig> EditorConfig = FindOrLoadConfig(EditorConfigName);
+	EditorConfig->SetRootUObject(*Object, Filter);
+	SaveConfig(EditorConfig);
+	return true;
+}
+
+bool UEditorConfigSubsystem::ReloadConfig(TSharedRef<FEditorConfig> Config)
+{
+	const FString* FilePath = LoadedConfigs.FindKey(Config);
+	if (!ensureMsgf(FilePath != nullptr, TEXT("Could not find filename for given config in UEditorConfigSubsystem::ReloadConfig().")))
+	{
+		return false;
+	}
+
+	const FString ConfigName = FPaths::GetBaseFilename(*FilePath);
+
+	TSharedPtr<FEditorConfig> Parent;
+
+	for (const FString& Dir : SearchDirectories)
+	{
+		const FString FullPath = FPaths::Combine(Dir, ConfigName) + TEXT(".json");
+
+		// find an existing config or create one
+		const TSharedPtr<FEditorConfig>* Existing = LoadedConfigs.Find(FullPath);
+		if (Existing == nullptr)
+		{
+			Existing = &LoadedConfigs.Add(FullPath, MakeShared<FEditorConfig>());
+		}
+
+		if (!(*Existing)->LoadFromFile(FullPath))
+		{
+			ensureMsgf(false, TEXT("Failed to load editor config from file %s"), *FullPath);
+			return false;
+		}
+
+		if (Parent.IsValid())
+		{
+			(*Existing)->SetParent(Parent);
+		}
+
+		Parent = (*Existing);
+	}
+
+	return true;
+}
+
+TSharedRef<FEditorConfig> UEditorConfigSubsystem::FindOrLoadConfig(FStringView ConfigName)
+{
+	checkf(!ConfigName.IsEmpty(), TEXT("Config name cannot be empty!"));
+
 	FString ConfigNameString(ConfigName);
-		
-	TSharedPtr<FEditorConfig> Current;
+	
+	// look for the config in the final search directory and return if it's loaded
+	// this assumes that the hierarchy of configs is unchanged
+	// ie. given search directories [Foo, Bar], the existence of Bar/X.json is taken to mean that Foo/X.json has been loaded
+	const FString FinalPath = FPaths::Combine(SearchDirectories.Last(), ConfigNameString) + TEXT(".json");
+
+	const TSharedPtr<FEditorConfig>* FinalConfig = LoadedConfigs.Find(FinalPath);
+	if (FinalConfig != nullptr)
+	{
+		return FinalConfig->ToSharedRef();
+	}
+	
+	// find or load all configs in all search directories with the given name
+	TSharedPtr<FEditorConfig> Parent;
 
 	for (const FString& Dir : SearchDirectories)
 	{
-		FString FullPath = FPaths::Combine(Dir, ConfigNameString) + TEXT(".json");
+		const FString FullPath = FPaths::Combine(Dir, ConfigNameString) + TEXT(".json");
 
 		const TSharedPtr<FEditorConfig>* Existing = LoadedConfigs.Find(FullPath);
 		if (Existing != nullptr && Existing->IsValid())
 		{
-			Current = *Existing;
+			Parent = *Existing;
 		}
 		else
 		{
 			// didn't exist yet, load now
-			TSharedPtr<FEditorConfig> NewConfig = MakeShared<FEditorConfig>();
+			TSharedRef<FEditorConfig> NewConfig = MakeShared<FEditorConfig>();
 			if (NewConfig->LoadFromFile(FullPath))
 			{
-				LoadedConfigs.Add(FullPath, NewConfig);
+				NewConfig->OnEditorConfigDirtied().AddUObject(this, &UEditorConfigSubsystem::OnEditorConfigDirtied);
 
-				if (Current.IsValid())
+				if (Parent.IsValid())
 				{
-					NewConfig->SetParent(Current);
+					NewConfig->SetParent(Parent);
 				}
 
-				Current = NewConfig;
+				LoadedConfigs.Add(FullPath, NewConfig);
+
+				Parent = NewConfig;
 			}
 		}
 	}
 
-	FString FinalPath = FPaths::Combine(SearchDirectories.Last(), ConfigNameString) + TEXT(".json");
-
-	const TSharedPtr<FEditorConfig>* FinalConfig = LoadedConfigs.Find(FinalPath);
+	FinalConfig = LoadedConfigs.Find(FinalPath);
 	if (FinalConfig != nullptr)
 	{
-		return *FinalConfig;
+		return FinalConfig->ToSharedRef();
 	}
 
 	// no config in the last search directory, create one now
 	// this will be the config that changes are written to
-	TSharedPtr<FEditorConfig> NewConfig = MakeShared<FEditorConfig>();
-	if (Current.IsValid())
+	TSharedRef<FEditorConfig> NewConfig = MakeShared<FEditorConfig>();
+	NewConfig->OnEditorConfigDirtied().AddUObject(this, &UEditorConfigSubsystem::OnEditorConfigDirtied);
+
+	if (Parent.IsValid())
 	{
 		// parent to the previous config
-		NewConfig->SetParent(Current);
+		NewConfig->SetParent(Parent);
 	}
 
 	LoadedConfigs.Add(FinalPath, NewConfig);
@@ -74,10 +189,21 @@ TSharedPtr<FEditorConfig> UEditorConfigSubsystem::FindOrLoadConfig(FStringView C
 	return NewConfig;
 }
 
-void UEditorConfigSubsystem::SaveConfig(TSharedPtr<FEditorConfig> Config, FOnCompletedDelegate OnCompleted)
+void UEditorConfigSubsystem::OnEditorConfigDirtied(const FEditorConfig& Config)
+{
+	for (const TPair<FString, TSharedPtr<FEditorConfig>>& Pair : LoadedConfigs)
+	{
+		if (Pair.Value.Get() == &Config)
+		{
+			SaveConfig(Pair.Value.ToSharedRef());
+		}
+	}
+}
+
+void UEditorConfigSubsystem::SaveConfig(TSharedRef<FEditorConfig> Config)
 {
 	const FString* FilePath = LoadedConfigs.FindKey(Config);
-	if (FilePath == nullptr)
+	if (!ensureMsgf(FilePath != nullptr, TEXT("Saving config that was not loaded through FEditorConfigSubsystem::FindOrLoadConfig. System does not know filepath to save to.")))
 	{
 		return;
 	}
@@ -85,29 +211,26 @@ void UEditorConfigSubsystem::SaveConfig(TSharedPtr<FEditorConfig> Config, FOnCom
 	SaveLock.WriteLock();
 	ON_SCOPE_EXIT { SaveLock.WriteUnlock(); };
 
-	for (const FPendingSave& Save : PendingSaves)
+	FPendingSave* Existing = PendingSaves.FindByPredicate([Config](const FPendingSave& Element)
+		{
+			return Element.Config == Config;
+		});
+
+	if (Existing != nullptr)
 	{
-		if (Save.Config == Config)
+		// reset the timer if we're saving within the grace period and no save is already being executed
+		if (!Existing->WasSuccess.IsValid())
 		{
-			// save already pending
-			return;
+			Existing->TimeSinceQueued = 0;
 		}
 	}
-
-	FString File = *FilePath;
-
-	FPendingSave& Save = PendingSaves.AddDefaulted_GetRef();
-	Save.Config = Config;
-	Save.FileName = File;
-	Save.OnCompleted = OnCompleted;
-	Save.WasSuccess = Async(EAsyncExecution::Thread, [Config, File]()
-		{
-			return Config->SaveToFile(File);
-		}, 
-		[this, Config]() 
-		{
-			OnSaveCompleted(Config);
-		});
+	else
+	{
+		FPendingSave& NewSave = PendingSaves.AddDefaulted_GetRef();
+		NewSave.Config = Config;
+		NewSave.FileName = *FilePath;
+		NewSave.TimeSinceQueued = 0;
+	}
 }
 
 void UEditorConfigSubsystem::OnSaveCompleted(TSharedPtr<FEditorConfig> Config)
@@ -115,7 +238,7 @@ void UEditorConfigSubsystem::OnSaveCompleted(TSharedPtr<FEditorConfig> Config)
 	SaveLock.WriteLock();
 	ON_SCOPE_EXIT { SaveLock.WriteUnlock(); };
 
-	int32 Index = PendingSaves.IndexOfByPredicate([Config](const FPendingSave& Element)
+	const int32 Index = PendingSaves.IndexOfByPredicate([Config](const FPendingSave& Element)
 		{
 			return Element.Config == Config;
 		});
@@ -123,7 +246,7 @@ void UEditorConfigSubsystem::OnSaveCompleted(TSharedPtr<FEditorConfig> Config)
 	if (Index != INDEX_NONE)
 	{
 		const FPendingSave& PendingSave = PendingSaves[Index];
-		PendingSave.OnCompleted.ExecuteIfBound(PendingSave.WasSuccess.Get());
+		PendingSave.Config->OnSaved();
 
 		PendingSaves.RemoveAt(Index);
 	}
diff --git a/Engine/Source/Editor/EditorConfig/Private/EditorMetadataOverrides.cpp b/Engine/Source/Editor/EditorConfig/Private/EditorMetadataOverrides.cpp
--- a/Engine/Source/Editor/EditorConfig/Private/EditorMetadataOverrides.cpp
+++ b/Engine/Source/Editor/EditorConfig/Private/EditorMetadataOverrides.cpp
@@ -32,25 +32,6 @@ bool UEditorMetadataOverrides::LoadFromConfig(TSharedPtr<FEditorConfig> Config)
 	return SourceConfig->TryGetStruct(TEXT("Metadata"), LoadedMetadata);
 }
 
-void UEditorMetadataOverrides::Tick(float DeltaTime)
-{
-	if (bDirty)
-	{
-		TimeSinceLastSave = TimeSinceLastSave + DeltaTime;
-
-		const float SaveDelaySeconds = 3.0f;
-		if (TimeSinceLastSave > SaveDelaySeconds)
-		{
-			Save();
-		}
-	}
-}
-
-TStatId UEditorMetadataOverrides::GetStatId() const
-{
-	RETURN_QUICK_DECLARE_CYCLE_STAT(UEditorMetadataOverrides, STATGROUP_Tickables);
-}
-
 void UEditorMetadataOverrides::Save()
 {
 	if (!SourceConfig.IsValid())
@@ -66,16 +47,7 @@ void UEditorMetadataOverrides::Save()
 		return;
 	}
 
-	EditorConfigSubsystem->SaveConfig(SourceConfig, FOnCompletedDelegate::CreateUObject(this, &UEditorMetadataOverrides::OnCompleted));
-}
-
-void UEditorMetadataOverrides::OnCompleted(bool bSuccess)
-{
-	if (bSuccess)
-	{
-		TimeSinceLastSave = 0;
-		bDirty = false;
-	}
+	EditorConfigSubsystem->SaveConfig(SourceConfig.ToSharedRef());
 }
 
 const FMetadataSet* UEditorMetadataOverrides::FindFieldMetadata(const FField* Field) const
@@ -242,7 +214,6 @@ void UEditorMetadataOverrides::SetStringMetadata(const FField* Field, FName Key,
 	}
 
 	FieldMetadata->Strings.Add(Key, FString(Value));
-	bDirty = true;
 }
 
 bool UEditorMetadataOverrides::GetFloatMetadata(const FField* Field, FName Key, float& OutValue) const
@@ -278,7 +249,7 @@ void UEditorMetadataOverrides::SetFloatMetadata(const FField* Field, FName Key,
 	}
 
 	FieldMetadata->Floats.Add(Key, Value);
-	bDirty = true;
+	Save();
 }
 
 bool UEditorMetadataOverrides::GetIntMetadata(const FField* Field, FName Key, int32& OutValue) const
@@ -314,7 +285,7 @@ void UEditorMetadataOverrides::SetIntMetadata(const FField* Field, FName Key, in
 	}
 
 	FieldMetadata->Ints.Add(Key, Value);
-	bDirty = true;
+	Save();
 }
 
 bool UEditorMetadataOverrides::GetBoolMetadata(const FField* Field, FName Key, bool& OutValue) const
@@ -350,7 +321,7 @@ void UEditorMetadataOverrides::SetBoolMetadata(const FField* Field, FName Key, b
 	}
 
 	FieldMetadata->Bools.Add(Key, Value);
-	bDirty = true;
+	Save();
 }
 
 bool UEditorMetadataOverrides::GetClassMetadata(const FField* Field, FName Key, UClass*& OutValue) const
@@ -374,7 +345,7 @@ void UEditorMetadataOverrides::SetClassMetadata(const FField* Field, FName Key,
 	}
 
 	SetStringMetadata(Field, Key, ClassName);
-	bDirty = true;
+	Save();
 }
 
 bool UEditorMetadataOverrides::GetArrayMetadata(const FField* Field, FName Key, TArray<FString>& OutValue) const
@@ -413,7 +384,7 @@ void UEditorMetadataOverrides::SetArrayMetadata(const FField* Field, FName Key,
 
 	const FString ValueString = FString::Join(Value, TEXT(","));
 	FieldMetadata->Strings.Add(Key, ValueString);
-	bDirty = true;
+	Save();
 }
 
 void UEditorMetadataOverrides::AddToArrayMetadata(const FField* Field, FName Key, const FString& Value)
@@ -434,7 +405,7 @@ void UEditorMetadataOverrides::AddToArrayMetadata(const FField* Field, FName Key
 		CurrentValue.Append(TEXT(","));
 		CurrentValue.Append(Value);
 	}
-	bDirty = true;
+	Save();
 }
 
 void UEditorMetadataOverrides::RemoveFromArrayMetadata(const FField* Field, FName Key, const FString& Value)
@@ -466,7 +437,7 @@ void UEditorMetadataOverrides::RemoveFromArrayMetadata(const FField* Field, FNam
 		FieldMetadata->Strings.Add(Key, ValueString);
 	}
 
-	bDirty = true;
+	Save();
 }
 
 void UEditorMetadataOverrides::RemoveMetadata(const FField* Field, FName Key)
@@ -499,7 +470,7 @@ void UEditorMetadataOverrides::RemoveMetadata(const FField* Field, FName Key)
 	FieldMetadata->Floats.Remove(Key);
 	FieldMetadata->Strings.Remove(Key);
 
-	bDirty = true;
+	Save();
 }
 
 const FMetadataSet* UEditorMetadataOverrides::FindStructMetadata(const UStruct* Struct) const
@@ -574,7 +545,7 @@ void UEditorMetadataOverrides::SetStringMetadata(const UStruct* Struct, FName Ke
 	}
 
 	StructMetadata->Strings.Add(Key, FString(Value));
-	bDirty = true;
+	Save();
 }
 
 bool UEditorMetadataOverrides::GetFloatMetadata(const UStruct* Struct, FName Key, float& OutValue) const
@@ -610,7 +581,7 @@ void UEditorMetadataOverrides::SetFloatMetadata(const UStruct* Struct, FName Key
 	}
 
 	StructMetadata->Floats.Add(Key, Value);
-	bDirty = true;
+	Save();
 }
 
 bool UEditorMetadataOverrides::GetIntMetadata(const UStruct* Struct, FName Key, int32& OutValue) const
@@ -646,7 +617,7 @@ void UEditorMetadataOverrides::SetIntMetadata(const UStruct* Struct, FName Key,
 	}
 
 	StructMetadata->Ints.Add(Key, Value);
-	bDirty = true;
+	Save();
 }
 
 bool UEditorMetadataOverrides::GetBoolMetadata(const UStruct* Struct, FName Key, bool& OutValue) const
@@ -682,7 +653,7 @@ void UEditorMetadataOverrides::SetBoolMetadata(const UStruct* Struct, FName Key,
 	}
 
 	StructMetadata->Bools.Add(Key, Value);
-	bDirty = true;
+	Save();
 }
 
 bool UEditorMetadataOverrides::GetClassMetadata(const UStruct* Struct, FName Key, UClass*& OutValue) const
@@ -706,7 +677,7 @@ void UEditorMetadataOverrides::SetClassMetadata(const UStruct* Struct, FName Key
 	}
 
 	SetStringMetadata(Struct, Key, ClassName);
-	bDirty = true;
+	Save();
 }
 
 bool UEditorMetadataOverrides::GetArrayMetadata(const UStruct* Struct, FName Key, TArray<FString>& OutValue) const
@@ -745,7 +716,7 @@ void UEditorMetadataOverrides::SetArrayMetadata(const UStruct* Struct, FName Key
 	const FString ValueString = FString::Join(Value, TEXT(","));
 
 	StructMetadata->Strings.Add(Key, ValueString);
-	bDirty = true;
+	Save();
 }
 
 void UEditorMetadataOverrides::AddToArrayMetadata(const UStruct* Struct, FName Key, const FString& Value)
@@ -767,7 +738,7 @@ void UEditorMetadataOverrides::AddToArrayMetadata(const UStruct* Struct, FName K
 		CurrentValue.Append(Value);
 	}
 
-	bDirty = true;
+	Save();
 }
 
 void UEditorMetadataOverrides::RemoveFromArrayMetadata(const UStruct* Struct, FName Key, const FString& Value)
@@ -799,7 +770,7 @@ void UEditorMetadataOverrides::RemoveFromArrayMetadata(const UStruct* Struct, FN
 		StructMetadata->Strings.Add(Key, ValueString);
 	}
 
-	bDirty = true;
+	Save();
 }
 
 void UEditorMetadataOverrides::RemoveMetadata(const UStruct* Struct, FName Key)
@@ -820,5 +791,5 @@ void UEditorMetadataOverrides::RemoveMetadata(const UStruct* Struct, FName Key)
 	StructMetadata->StructMetadata.Floats.Remove(Key);
 	StructMetadata->StructMetadata.Strings.Remove(Key);
 
-	bDirty = true;
+	Save();
 }
diff --git a/Engine/Source/Editor/EditorConfig/Public/EditorConfig.h b/Engine/Source/Editor/EditorConfig/Public/EditorConfig.h
--- a/Engine/Source/Editor/EditorConfig/Public/EditorConfig.h
+++ b/Engine/Source/Editor/EditorConfig/Public/EditorConfig.h
@@ -10,6 +10,12 @@
 class EDITORCONFIG_API FEditorConfig
 {
 public:
+	enum class EPropertyFilter
+	{
+		All,
+		MetadataOnly
+	};
+
 	FEditorConfig();
 
 	void SetParent(TSharedPtr<FEditorConfig> InConfig);
@@ -23,37 +29,45 @@ public:
 
 	// UStruct & UObject
 	template <typename T>
-	bool TryGetStruct(FStringView Key, T& OutValue) const;
+	bool TryGetStruct(FStringView Key, T& OutValue, EPropertyFilter Filter = EPropertyFilter::MetadataOnly) const;
 	template <typename T>
-	bool TryGetUObject(FStringView Key, T& OutValue) const;
+	bool TryGetUObject(FStringView Key, T& OutValue, EPropertyFilter Filter = EPropertyFilter::MetadataOnly) const;
 
 	template <typename T>
-	bool TryGetRootStruct(T& OutValue) const;
+	bool TryGetRootStruct(T& OutValue, EPropertyFilter Filter = EPropertyFilter::MetadataOnly) const;
 	template <typename T>
-	bool TryGetRootUObject(T& OutValue) const;
+	bool TryGetRootUObject(T& OutValue, EPropertyFilter Filter = EPropertyFilter::MetadataOnly) const;
 
 	template <typename T>
-	void SetStruct(FStringView Key, const T& InValue);
+	void SetStruct(FStringView Key, const T& InValue, EPropertyFilter Filter = EPropertyFilter::MetadataOnly);
 	template <typename T>
-	void SetUObject(FStringView Key, const T& InValue);
+	void SetUObject(FStringView Key, const T& InValue, EPropertyFilter Filter = EPropertyFilter::MetadataOnly);
 
 	template <typename T>
-	void SetRootStruct(const T& InValue);
+	void SetRootStruct(const T& InValue, EPropertyFilter Filter = EPropertyFilter::MetadataOnly);
 	template <typename T>
-	void SetRootUObject(const T& InValue);
+	void SetRootUObject(const T& InValue, EPropertyFilter Filter = EPropertyFilter::MetadataOnly);
+
+	void SetRootUObject(const UClass* Class, const UObject* Instance, EPropertyFilter Filter = EPropertyFilter::MetadataOnly);
+	void SetRootStruct(const UStruct* Class, const void* Instance, EPropertyFilter Filter = EPropertyFilter::MetadataOnly);
 
 	bool HasOverride(FStringView Key) const;
 
+	void OnSaved();
+
+	DECLARE_EVENT_OneParam(FEditorConfig, FOnEditorConfigDirtied, const FEditorConfig&);
+	FOnEditorConfigDirtied& OnEditorConfigDirtied() { return EditorConfigDirtiedEvent; }
+
 private:
 
 	friend class UEditorConfigSubsystem; // for access to LoadFromFile and SaveToFile
 
-	static void ReadUObject(TSharedPtr<FJsonObject> JsonObject, const UClass* Class, UObject* Instance);
-	static void ReadStruct(TSharedPtr<FJsonObject> JsonObject, const UStruct* Struct, void* Instance, UObject* Owner);
+	static void ReadUObject(TSharedPtr<FJsonObject> JsonObject, const UClass* Class, UObject* Instance, EPropertyFilter Filter);
+	static void ReadStruct(TSharedPtr<FJsonObject> JsonObject, const UStruct* Struct, void* Instance, UObject* Owner, EPropertyFilter Filter);
 	static void ReadValue(TSharedPtr<FJsonValue> JsonValue, const FProperty* Property, void* DataPtr, UObject* Owner);
 	
-	static TSharedPtr<FJsonObject> WriteStruct(const UStruct* Struct, const void* Instance);
-	static TSharedPtr<FJsonObject> WriteUObject(const UStruct* Struct, const UObject* Instance);
+	static TSharedPtr<FJsonObject> WriteStruct(const UStruct* Struct, const void* Instance, EPropertyFilter Filter);
+	static TSharedPtr<FJsonObject> WriteUObject(const UClass* Class, const UObject* Instance, EPropertyFilter Filter);
 	static TSharedPtr<FJsonValue> WriteValue(const FProperty* Property, const void* DataPtr);
 
 	static bool IsDefault(const FProperty* Property, TSharedPtr<FJsonValue> JsonValue, const void* NativeValue);
@@ -61,15 +75,18 @@ private:
 	bool LoadFromFile(FStringView FilePath);
 	bool SaveToFile(FStringView FilePath) const;
 
+	void SetDirty();
+
 private:
 	TSharedPtr<UE::FJsonConfig> JsonConfig;
 	TSharedPtr<FEditorConfig> ParentConfig;
 		
+	FOnEditorConfigDirtied EditorConfigDirtiedEvent;
 	bool Dirty { false };
 };
 
 template <typename T>
-bool FEditorConfig::TryGetStruct(FStringView Key, T& OutValue) const
+bool FEditorConfig::TryGetStruct(FStringView Key, T& OutValue, EPropertyFilter Filter) const
 {
 	if (!IsValid())
 	{
@@ -88,14 +105,14 @@ bool FEditorConfig::TryGetStruct(FStringView Key, T& OutValue) const
 		return false;
 	}
 
-	UStruct* Struct = T::StaticStruct();
-	ReadStruct(StructData, Struct, &OutValue, nullptr);
+	const UStruct* Struct = T::StaticStruct();
+	ReadStruct(StructData, Struct, &OutValue, nullptr, Filter);
 
 	return true;
 }
 
 template <typename T>
-bool FEditorConfig::TryGetUObject(FStringView Key, T& OutValue) const
+bool FEditorConfig::TryGetUObject(FStringView Key, T& OutValue, EPropertyFilter Filter) const
 {
 	static_assert(TIsDerivedFrom<T, UObject>::Value, "Type is not derived from UObject.");
 
@@ -116,14 +133,14 @@ bool FEditorConfig::TryGetUObject(FStringView Key, T& OutValue) const
 		return false;
 	}
 
-	UClass* Class = T::StaticClass();
-	ReadUObject(UObjectData, Class, &OutValue);
+	const UClass* Class = T::StaticClass();
+	ReadUObject(UObjectData, Class, &OutValue, Filter);
 
 	return true;
 }
 
 template <typename T>
-bool FEditorConfig::TryGetRootStruct(T& OutValue) const
+bool FEditorConfig::TryGetRootStruct(T& OutValue, EPropertyFilter Filter) const
 {
 	if (!IsValid())
 	{
@@ -132,14 +149,14 @@ bool FEditorConfig::TryGetRootStruct(T& OutValue) const
 
 	TSharedPtr<FJsonObject> StructData = JsonConfig->GetRootObject();
 
-	UStruct* Struct = T::StaticStruct();
-	ReadStruct(StructData, Struct, &OutValue, nullptr);
+	const UStruct* Struct = T::StaticStruct();
+	ReadStruct(StructData, Struct, &OutValue, nullptr, Filter);
 
 	return true;
 }
 
 template <typename T>
-bool FEditorConfig::TryGetRootUObject(T& OutValue) const
+bool FEditorConfig::TryGetRootUObject(T& OutValue, EPropertyFilter Filter) const
 {
 	static_assert(TIsDerivedFrom<T, UObject>::Value, "Type is not derived from UObject.");
 
@@ -150,28 +167,28 @@ bool FEditorConfig::TryGetRootUObject(T& OutValue) const
 
 	TSharedPtr<FJsonObject> UObjectData = JsonConfig->GetRootObject();
 
-	UClass* Class = T::StaticClass();
-	ReadUObject(UObjectData, Class, &OutValue);
+	const UClass* Class = T::StaticClass();
+	ReadUObject(UObjectData, Class, &OutValue, Filter);
 
 	return true;
 }
 
 template <typename T>
-void FEditorConfig::SetStruct(FStringView Key, const T& InValue)
+void FEditorConfig::SetStruct(FStringView Key, const T& InValue, EPropertyFilter Filter)
 {
 	if (!IsValid())
 	{
 		return;
 	}
 
-	TSharedPtr<FJsonObject> JsonObject = WriteStruct(T::StaticStruct(), &InValue);
+	TSharedPtr<FJsonObject> JsonObject = WriteStruct(T::StaticStruct(), &InValue, Filter);
 	JsonConfig->SetJsonObject(UE::FJsonPath(Key), JsonObject);
 
-	Dirty = true;
+	SetDirty();
 }
 
 template <typename T>
-void FEditorConfig::SetUObject(FStringView Key, const T& InValue)
+void FEditorConfig::SetUObject(FStringView Key, const T& InValue, EPropertyFilter Filter)
 {
 	static_assert(TIsDerivedFrom<T, UObject>::Value, "Type is not derived from UObject.");
 
@@ -180,38 +197,22 @@ void FEditorConfig::SetUObject(FStringView Key, const T& InValue)
 		return;
 	}
 	
-	TSharedPtr<FJsonObject> JsonObject = WriteStruct(T::StaticClass(), &InValue);
+	TSharedPtr<FJsonObject> JsonObject = WriteUObject(T::StaticClass(), &InValue);
 	JsonConfig->SetJsonObject(UE::FJsonPath(Key), JsonObject);
 		
-	Dirty = true;
+	SetDirty();
 }
 
 template <typename T>
-void FEditorConfig::SetRootStruct(const T& InValue)
+void FEditorConfig::SetRootStruct(const T& InValue, EPropertyFilter Filter)
 {
-	if (!IsValid())
-	{
-		return;
-	}
-
-	TSharedPtr<FJsonObject> JsonObject = WriteStruct(T::StaticStruct(), &InValue);
-	JsonConfig->SetRootObject(JsonObject);
-		
-	Dirty = true;
+	SetRootStruct(T::StaticStruct(), &InValue, Filter);
 }
 
 template <typename T>
-void FEditorConfig::SetRootUObject(const T& InValue)
+void FEditorConfig::SetRootUObject(const T& InValue, EPropertyFilter Filter)
 {
 	static_assert(TIsDerivedFrom<T, UObject>::Value, "Type is not derived from UObject.");
 
-	if (!IsValid())
-	{
-		return;
-	}
-	
-	TSharedPtr<FJsonObject> JsonObject = WriteStruct(T::StaticStruct(), &InValue);
-	JsonConfig->SetRootObject(JsonObject);
-		
-	Dirty = true;
+	SetRootUObject(T::StaticClass(), &InValue, Filter);
 }
diff --git a/Engine/Source/Editor/EditorConfig/Public/EditorConfigBase.h b/Engine/Source/Editor/EditorConfig/Public/EditorConfigBase.h
new file mode 100644
--- /dev/null
+++ b/Engine/Source/Editor/EditorConfig/Public/EditorConfigBase.h
@@ -0,0 +1,23 @@
+// Copyright Epic Games, Inc. All Rights Reserved.
+
+#pragma once
+
+#include "UObject/Object.h"
+#include "UObject/ObjectMacros.h"
+
+#include "EditorConfigBase.generated.h"
+
+/** Inherit from this class to simplify saving and loading properties from editor configs. */
+UCLASS()
+class EDITORCONFIG_API UEditorConfigBase : public UObject
+{
+	GENERATED_BODY()
+
+public:
+
+	/** Load any properties of this class into properties marked with metadata tag "EditorConfig" from the class's EditorConfig */
+	bool LoadEditorConfig();
+
+	/** Save any properties of this class in properties marked with metadata tag "EditorConfig" into the class's EditorConfig. */
+	bool SaveEditorConfig() const;
+};
\ No newline at end of file
diff --git a/Engine/Source/Editor/EditorConfig/Public/EditorConfigSubsystem.h b/Engine/Source/Editor/EditorConfig/Public/EditorConfigSubsystem.h
--- a/Engine/Source/Editor/EditorConfig/Public/EditorConfigSubsystem.h
+++ b/Engine/Source/Editor/EditorConfig/Public/EditorConfigSubsystem.h
@@ -2,31 +2,79 @@
 
 #pragma once
 
+#include "Editor.h"
 #include "EditorConfig.h"
 #include "EditorSubsystem.h"
+#include "TickableEditorObject.h"
 
 #include "EditorConfigSubsystem.generated.h"
 
-DECLARE_DELEGATE_OneParam(FOnCompletedDelegate, bool);
-
 UCLASS()
 class EDITORCONFIG_API UEditorConfigSubsystem : 
-	public UEditorSubsystem
+	public UEditorSubsystem,
+	public FTickableEditorObject
 {
 	GENERATED_BODY()
 
 public:
 	UEditorConfigSubsystem();
 
-	void Initialize(FSubsystemCollectionBase& Collection);
+	/** 
+	 * Load a config from the root of the JSON file into a given UObject.
+	 * This loads the config from the UCLASS's EditorConfig="ConfigName" metadata.
+	 * @param Filter Whether to load all properties, or only ones marked with the EditorConfig metadata.
+	 */
+	template <typename TObject>
+	bool LoadConfigObject(TObject* Object, FEditorConfig::EPropertyFilter Filter = FEditorConfig::EPropertyFilter::MetadataOnly);
 
-	TSharedPtr<FEditorConfig> FindOrLoadConfig(FStringView ConfigName);
-	void SaveConfig(TSharedPtr<FEditorConfig> Config, FOnCompletedDelegate OnCompleted);
+	/** 
+	 * Save the given UObject to the root of the JSON config.
+	 * This saves the config to the UCLASS's EditorConfig="ConfigName" metadata.
+	 * @param Filter Whether to save all properties, or only ones marked with the EditorConfig metadata.
+	 */
+	template <typename TObject>
+	bool SaveConfigObject(const TObject* Object, FEditorConfig::EPropertyFilter Filter = FEditorConfig::EPropertyFilter::MetadataOnly);
 
-	void AddSearchDirectory(FStringView SearchDir);
+	/** 
+	 * Load a config from the root of the JSON file into a given UObject.
+	 * This loads the config from the UCLASS's EditorConfig="ConfigName" metadata.
+	 * @param Filter Whether to load all properties, or only ones marked with the EditorConfig metadata.
+	 * @param Object The UObject instance to load into.
+	 * @param Class The UClass of the object. 
+	 */
+	bool LoadConfigObject(UObject* Object, const UClass* Class, FEditorConfig::EPropertyFilter = FEditorConfig::EPropertyFilter::MetadataOnly);
+	
+	/** 
+	 * Save the given UObject of the given class to the root of the JSON config.
+	 * This saves the config to the UCLASS's EditorConfig="ConfigName" metadata.
+	 * @param Object The UObject instance to save.
+	 * @param Class The UClass of the object. 
+	 * @param Filter Whether to save all properties, or only ones marked with the EditorConfig metadata.
+	 */
+	bool SaveConfigObject(const UObject* Object, const UClass* Class, FEditorConfig::EPropertyFilter = FEditorConfig::EPropertyFilter::MetadataOnly);
+	
+	/** 
+	 * Find a config with the given name that has already been loaded, load it if it hasn't been, or create one with the given name.
+	 */
+	TSharedRef<FEditorConfig> FindOrLoadConfig(FStringView ConfigName);
+
+	/*
+	 * Save the given config to the location it was loaded.
+	 */
+	void SaveConfig(TSharedRef<FEditorConfig> Config);
+
+	/** Force reload the given config and all its (current and potential) parents from disk. */
+	bool ReloadConfig(TSharedRef<FEditorConfig> Config);
 
 private:
+	void Initialize(FSubsystemCollectionBase& Collection) override;
+	void AddSearchDirectory(FStringView SearchDir);
 	void OnSaveCompleted(TSharedPtr<FEditorConfig> Config);
+	void OnEditorConfigDirtied(const FEditorConfig& Config);
+
+	/** FTickableEditorObject interface */
+	void Tick(float DeltaTime) override;
+	TStatId GetStatId() const override;
 
 private:
 	struct FPendingSave
@@ -34,7 +82,7 @@ private:
 		FString FileName;
 		TSharedPtr<FEditorConfig> Config;
 		TFuture<bool> WasSuccess;
-		FOnCompletedDelegate OnCompleted;
+		float TimeSinceQueued { 0 };
 	};
 
 	FRWLock SaveLock;
@@ -42,3 +90,23 @@ private:
 	TArray<FString> SearchDirectories;
 	TMap<FString, TSharedPtr<FEditorConfig>> LoadedConfigs;
 };
+
+template <typename TObject>
+bool UEditorConfigSubsystem::LoadConfigObject(TObject* Object, FEditorConfig::EPropertyFilter Filter)
+{
+	static_assert(TIsDerivedFrom<TObject, UObject>::Value, "Type is not derived from UObject.");
+	check(Object != nullptr);
+	
+	const UClass* Class = TObject::StaticClass();
+	return LoadConfigObject(static_cast<UObject*>(Object), Class, Filter);
+}
+
+template <typename TObject>
+bool UEditorConfigSubsystem::SaveConfigObject(const TObject* Object, FEditorConfig::EPropertyFilter Filter)
+{
+	static_assert(TIsDerivedFrom<TObject, UObject>::Value, "Type is not derived from UObject.");
+	check(Object != nullptr);
+
+	const UClass* Class = TObject::StaticClass();
+	return SaveConfigObject(static_cast<const UObject*>(Object), Class);
+}
\ No newline at end of file
diff --git a/Engine/Source/Editor/EditorConfig/Public/EditorMetadataOverrides.h b/Engine/Source/Editor/EditorConfig/Public/EditorMetadataOverrides.h
--- a/Engine/Source/Editor/EditorConfig/Public/EditorMetadataOverrides.h
+++ b/Engine/Source/Editor/EditorConfig/Public/EditorMetadataOverrides.h
@@ -5,7 +5,6 @@
 #include "CoreMinimal.h"
 #include "EditorConfig.h"
 #include "EditorSubsystem.h"
-#include "TickableEditorObject.h"
 
 #include "EditorMetadataOverrides.generated.h"
 
@@ -63,8 +62,7 @@ enum class EMetadataType
 
 UCLASS()
 class EDITORCONFIG_API UEditorMetadataOverrides : 
-	public UEditorSubsystem, 
-	public FTickableEditorObject
+	public UEditorSubsystem
 { 
 	GENERATED_BODY()
 
@@ -126,9 +124,6 @@ public:
 
 	void RemoveMetadata(const UStruct* Struct, FName Key);
 
-	void Tick(float DeltaTime) override;
-	TStatId GetStatId() const override;
-
 private:
 	const FMetadataSet* FindFieldMetadata(const FField* Field) const;
 	FMetadataSet* FindFieldMetadata(const FField* Field);
@@ -143,7 +138,4 @@ private:
 private:
 	TSharedPtr<FEditorConfig> SourceConfig;
 	FMetadataConfig LoadedMetadata;
-
-	TAtomic<bool> bDirty { false };
-	TAtomic<float> TimeSinceLastSave { 0 };
 };
diff --git a/Engine/Source/Programs/UnrealHeaderTool/Private/Specifiers/ClassMetadataSpecifiers.def b/Engine/Source/Programs/UnrealHeaderTool/Private/Specifiers/ClassMetadataSpecifiers.def
--- a/Engine/Source/Programs/UnrealHeaderTool/Private/Specifiers/ClassMetadataSpecifiers.def
+++ b/Engine/Source/Programs/UnrealHeaderTool/Private/Specifiers/ClassMetadataSpecifiers.def
@@ -23,6 +23,7 @@ CLASS_METADATA_SPECIFIER(Deprecated)
 CLASS_METADATA_SPECIFIER(DontAutoCollapseCategories)
 CLASS_METADATA_SPECIFIER(DontCollapseCategories)
 CLASS_METADATA_SPECIFIER(EditInlineNew)
+CLASS_METADATA_SPECIFIER(EditorConfig)
 CLASS_METADATA_SPECIFIER(GlobalUserConfig)
 CLASS_METADATA_SPECIFIER(HideCategories)
 CLASS_METADATA_SPECIFIER(HideDropdown)
diff --git a/Engine/Source/Programs/UnrealHeaderTool/Private/UnrealTypeDefinitionInfo.cpp b/Engine/Source/Programs/UnrealHeaderTool/Private/UnrealTypeDefinitionInfo.cpp
--- a/Engine/Source/Programs/UnrealHeaderTool/Private/UnrealTypeDefinitionInfo.cpp
+++ b/Engine/Source/Programs/UnrealHeaderTool/Private/UnrealTypeDefinitionInfo.cpp
@@ -30,6 +30,7 @@ namespace
 	const FName NAME_IgnoreCategoryKeywordsInSubclasses(TEXT("IgnoreCategoryKeywordsInSubclasses"));
 	const FName NAME_ReplaceConverted(TEXT("ReplaceConverted"));
 	const FName NAME_ClassGroupNames(TEXT("ClassGroupNames"));
+	const FName NAME_EditorConfig(TEXT("EditorConfig"));
 
 	/**
 	 * As part of the singleton name, collect the parent chain names
@@ -1088,6 +1089,11 @@ void FUnrealClassDefinitionInfo::ParseClassProperties(TArray<FPropertySpecifier>
 			ParsedClassFlags |= CLASS_ProjectUserConfig;
 			break;
 
+		case EClassMetadataSpecifier::EditorConfig:
+			// Save EditorConfig properties to the given JSON file.
+			MetaData.Add(NAME_EditorConfig, FHeaderParser::RequireExactlyOneSpecifierValue(PropSpecifier));
+			break;
+
 		case EClassMetadataSpecifier::ShowCategories:
 
 			FHeaderParser::RequireSpecifierValue(PropSpecifier);
diff --git a/Engine/Source/Runtime/CoreUObject/Public/UObject/ObjectMacros.h b/Engine/Source/Runtime/CoreUObject/Public/UObject/ObjectMacros.h
--- a/Engine/Source/Runtime/CoreUObject/Public/UObject/ObjectMacros.h
+++ b/Engine/Source/Runtime/CoreUObject/Public/UObject/ObjectMacros.h
@@ -726,6 +726,9 @@ namespace UC
 		/// Save object config only to Default INIs, never to local INIs.
 		defaultconfig,
 
+		/// Mark the editor config file to load from if loading into this object.
+		EditorConfig,
+
 		/// These affect the behavior of the property editor.
 		/// Class can be constructed from editinline New button.
 		editinlinenew,
