Added UClass & UStruct versions of FEditorConfigSubsystem::TryGetRoot*.

diff --git a/Engine/Source/Editor/EditorConfig/Private/EditorConfig.cpp b/Engine/Source/Editor/EditorConfig/Private/EditorConfig.cpp
--- a/Engine/Source/Editor/EditorConfig/Private/EditorConfig.cpp
+++ b/Engine/Source/Editor/EditorConfig/Private/EditorConfig.cpp
@@ -81,6 +81,32 @@ bool FEditorConfig::HasOverride(FStringView Key) const
 	return JsonConfig->HasOverride(UE::FJsonPath(Key));
 }
 
+bool FEditorConfig::TryGetRootUObject(const UClass* Class, UObject* OutValue, EPropertyFilter Filter) const
+{
+	if (!IsValid())
+	{
+		return false;
+	}
+
+	TSharedPtr<FJsonObject> UObjectData = JsonConfig->GetRootObject();
+	ReadUObject(UObjectData, Class, OutValue, Filter);
+
+	return true;
+}
+
+bool FEditorConfig::TryGetRootStruct(const UStruct* Struct, void* OutValue, EPropertyFilter Filter) const
+{
+	if (!IsValid())
+	{
+		return false;
+	}
+
+	TSharedPtr<FJsonObject> StructData = JsonConfig->GetRootObject();
+	ReadStruct(StructData, Struct, OutValue, nullptr, Filter);
+
+	return true;
+}
+
 void FEditorConfig::SetRootUObject(const UClass* Class, const UObject* Instance, EPropertyFilter Filter)
 {
 	if (!IsValid())
diff --git a/Engine/Source/Editor/EditorConfig/Private/EditorConfigBase.cpp b/Engine/Source/Editor/EditorConfig/Private/EditorConfigBase.cpp
--- a/Engine/Source/Editor/EditorConfig/Private/EditorConfigBase.cpp
+++ b/Engine/Source/Editor/EditorConfig/Private/EditorConfigBase.cpp
@@ -7,11 +7,11 @@
 bool UEditorConfigBase::LoadEditorConfig()
 {
 	UEditorConfigSubsystem* Subsystem = GEditor->GetEditorSubsystem<UEditorConfigSubsystem>();
-	return Subsystem->LoadConfigObject(this, this->GetClass());
+	return Subsystem->LoadConfigObject(this->GetClass(), this);
 }
 
 bool UEditorConfigBase::SaveEditorConfig() const
 {
 	UEditorConfigSubsystem* Subsystem = GEditor->GetEditorSubsystem<UEditorConfigSubsystem>();
-	return Subsystem->SaveConfigObject(this, this->GetClass());
+	return Subsystem->SaveConfigObject(this->GetClass(), this);
 }
\ No newline at end of file
diff --git a/Engine/Source/Editor/EditorConfig/Private/EditorConfigSubystem.cpp b/Engine/Source/Editor/EditorConfig/Private/EditorConfigSubystem.cpp
--- a/Engine/Source/Editor/EditorConfig/Private/EditorConfigSubystem.cpp
+++ b/Engine/Source/Editor/EditorConfig/Private/EditorConfigSubystem.cpp
@@ -52,7 +52,7 @@ TStatId UEditorConfigSubsystem::GetStatId() const
 	RETURN_QUICK_DECLARE_CYCLE_STAT(UEditorConfigSubsystem, STATGROUP_Tickables);
 }
 
-bool UEditorConfigSubsystem::LoadConfigObject(UObject* Object, const UClass* Class, FEditorConfig::EPropertyFilter Filter)
+bool UEditorConfigSubsystem::LoadConfigObject(const UClass* Class, UObject* Object, FEditorConfig::EPropertyFilter Filter)
 {
 	const FString& EditorConfigName = Class->GetMetaData("EditorConfig");
 	if (!ensureMsgf(!EditorConfigName.IsEmpty(), TEXT("UEditorConfigSubsystem::LoadConfigObject - EditorConfig name is not set on class %s."), *Class->GetName()))
@@ -61,19 +61,19 @@ bool UEditorConfigSubsystem::LoadConfigObject(UObject* Object, const UClass* Cla
 	}
 
 	TSharedRef<FEditorConfig> EditorConfig = FindOrLoadConfig(EditorConfigName);
-	return EditorConfig->TryGetRootUObject(*Object, Filter);
+	return EditorConfig->TryGetRootUObject(Class, Object, Filter);
 }
 
-bool UEditorConfigSubsystem::SaveConfigObject(const UObject* Object, const UClass* Class, FEditorConfig::EPropertyFilter Filter)
+bool UEditorConfigSubsystem::SaveConfigObject(const UClass* Class, const UObject* Object, FEditorConfig::EPropertyFilter Filter)
 {
 	const FString& EditorConfigName = Class->GetMetaData("EditorConfig");
-	if (!ensureMsgf(EditorConfigName.IsEmpty(), TEXT("UEditorConfigSubsystem::SaveConfigObject - EditorConfig name is not set on class %s."), *Class->GetName()))
+	if (!ensureMsgf(!EditorConfigName.IsEmpty(), TEXT("UEditorConfigSubsystem::SaveConfigObject - EditorConfig name is not set on class %s."), *Class->GetName()))
 	{
 		return false;
 	}
 
 	TSharedRef<FEditorConfig> EditorConfig = FindOrLoadConfig(EditorConfigName);
-	EditorConfig->SetRootUObject(*Object, Filter);
+	EditorConfig->SetRootUObject(Class, Object, Filter);
 	SaveConfig(EditorConfig);
 	return true;
 }
diff --git a/Engine/Source/Editor/EditorConfig/Public/EditorConfig.h b/Engine/Source/Editor/EditorConfig/Public/EditorConfig.h
--- a/Engine/Source/Editor/EditorConfig/Public/EditorConfig.h
+++ b/Engine/Source/Editor/EditorConfig/Public/EditorConfig.h
@@ -38,6 +38,9 @@ public:
 	template <typename T>
 	bool TryGetRootUObject(T& OutValue, EPropertyFilter Filter = EPropertyFilter::MetadataOnly) const;
 
+	bool TryGetRootStruct(const UStruct* Class, void* OutValue, EPropertyFilter Filter = EPropertyFilter::MetadataOnly) const;
+	bool TryGetRootUObject(const UClass* Class, UObject* OutValue, EPropertyFilter Filter = EPropertyFilter::MetadataOnly) const;
+
 	template <typename T>
 	void SetStruct(FStringView Key, const T& InValue, EPropertyFilter Filter = EPropertyFilter::MetadataOnly);
 	template <typename T>
@@ -48,8 +51,8 @@ public:
 	template <typename T>
 	void SetRootUObject(const T& InValue, EPropertyFilter Filter = EPropertyFilter::MetadataOnly);
 
-	void SetRootUObject(const UClass* Class, const UObject* Instance, EPropertyFilter Filter = EPropertyFilter::MetadataOnly);
 	void SetRootStruct(const UStruct* Class, const void* Instance, EPropertyFilter Filter = EPropertyFilter::MetadataOnly);
+	void SetRootUObject(const UClass* Class, const UObject* Instance, EPropertyFilter Filter = EPropertyFilter::MetadataOnly);
 
 	bool HasOverride(FStringView Key) const;
 
@@ -142,17 +145,7 @@ bool FEditorConfig::TryGetUObject(FStringView Key, T& OutValue, EPropertyFilter
 template <typename T>
 bool FEditorConfig::TryGetRootStruct(T& OutValue, EPropertyFilter Filter) const
 {
-	if (!IsValid())
-	{
-		return false;
-	}
-
-	TSharedPtr<FJsonObject> StructData = JsonConfig->GetRootObject();
-
-	const UStruct* Struct = T::StaticStruct();
-	ReadStruct(StructData, Struct, &OutValue, nullptr, Filter);
-
-	return true;
+	return TryGetRootStruct(T::StaticStruct(), &OutValue, Filter);
 }
 
 template <typename T>
@@ -160,17 +153,7 @@ bool FEditorConfig::TryGetRootUObject(T& OutValue, EPropertyFilter Filter) const
 {
 	static_assert(TIsDerivedFrom<T, UObject>::Value, "Type is not derived from UObject.");
 
-	if (!IsValid())
-	{
-		return false;
-	}
-
-	TSharedPtr<FJsonObject> UObjectData = JsonConfig->GetRootObject();
-
-	const UClass* Class = T::StaticClass();
-	ReadUObject(UObjectData, Class, &OutValue, Filter);
-
-	return true;
+	return TryGetRootUObject(T::StaticClass(), &OutValue, Filter);
 }
 
 template <typename T>
diff --git a/Engine/Source/Editor/EditorConfig/Public/EditorConfigSubsystem.h b/Engine/Source/Editor/EditorConfig/Public/EditorConfigSubsystem.h
--- a/Engine/Source/Editor/EditorConfig/Public/EditorConfigSubsystem.h
+++ b/Engine/Source/Editor/EditorConfig/Public/EditorConfigSubsystem.h
@@ -22,6 +22,7 @@ public:
 	/** 
 	 * Load a config from the root of the JSON file into a given UObject.
 	 * This loads the config from the UCLASS's EditorConfig="ConfigName" metadata.
+	 * @param Object The object to load into.
 	 * @param Filter Whether to load all properties, or only ones marked with the EditorConfig metadata.
 	 */
 	template <typename TObject>
@@ -30,6 +31,7 @@ public:
 	/** 
 	 * Save the given UObject to the root of the JSON config.
 	 * This saves the config to the UCLASS's EditorConfig="ConfigName" metadata.
+	 * @param Object The object to save.
 	 * @param Filter Whether to save all properties, or only ones marked with the EditorConfig metadata.
 	 */
 	template <typename TObject>
@@ -38,20 +40,20 @@ public:
 	/** 
 	 * Load a config from the root of the JSON file into a given UObject.
 	 * This loads the config from the UCLASS's EditorConfig="ConfigName" metadata.
-	 * @param Filter Whether to load all properties, or only ones marked with the EditorConfig metadata.
-	 * @param Object The UObject instance to load into.
 	 * @param Class The UClass of the object. 
+	 * @param Object The UObject instance to load into.
+	 * @param Filter Whether to load all properties, or only ones marked with the EditorConfig metadata.
 	 */
-	bool LoadConfigObject(UObject* Object, const UClass* Class, FEditorConfig::EPropertyFilter = FEditorConfig::EPropertyFilter::MetadataOnly);
+	bool LoadConfigObject(const UClass* Class, UObject* Object, FEditorConfig::EPropertyFilter = FEditorConfig::EPropertyFilter::MetadataOnly);
 	
 	/** 
 	 * Save the given UObject of the given class to the root of the JSON config.
 	 * This saves the config to the UCLASS's EditorConfig="ConfigName" metadata.
-	 * @param Object The UObject instance to save.
 	 * @param Class The UClass of the object. 
+	 * @param Object The UObject instance to save.
 	 * @param Filter Whether to save all properties, or only ones marked with the EditorConfig metadata.
 	 */
-	bool SaveConfigObject(const UObject* Object, const UClass* Class, FEditorConfig::EPropertyFilter = FEditorConfig::EPropertyFilter::MetadataOnly);
+	bool SaveConfigObject(const UClass* Class, const UObject* Object, FEditorConfig::EPropertyFilter = FEditorConfig::EPropertyFilter::MetadataOnly);
 	
 	/** 
 	 * Find a config with the given name that has already been loaded, load it if it hasn't been, or create one with the given name.
@@ -98,7 +100,7 @@ bool UEditorConfigSubsystem::LoadConfigObject(TObject* Object, FEditorConfig::EP
 	check(Object != nullptr);
 	
 	const UClass* Class = TObject::StaticClass();
-	return LoadConfigObject(static_cast<UObject*>(Object), Class, Filter);
+	return LoadConfigObject(Class, Object, Filter);
 }
 
 template <typename TObject>
@@ -108,5 +110,5 @@ bool UEditorConfigSubsystem::SaveConfigObject(const TObject* Object, FEditorConf
 	check(Object != nullptr);
 
 	const UClass* Class = TObject::StaticClass();
-	return SaveConfigObject(static_cast<const UObject*>(Object), Class);
+	return SaveConfigObject(Class, Object);
 }
\ No newline at end of file
