Speculative fix for crash when cooking and GetEditorSubsystem() returns nullptr.

diff --git a/Engine/Source/Editor/EditorConfig/Private/EditorConfigBase.cpp b/Engine/Source/Editor/EditorConfig/Private/EditorConfigBase.cpp
--- a/Engine/Source/Editor/EditorConfig/Private/EditorConfigBase.cpp
+++ b/Engine/Source/Editor/EditorConfig/Private/EditorConfigBase.cpp
@@ -7,11 +7,21 @@
 bool UEditorConfigBase::LoadEditorConfig()
 {
 	UEditorConfigSubsystem* Subsystem = GEditor->GetEditorSubsystem<UEditorConfigSubsystem>();
-	return Subsystem->LoadConfigObject(this->GetClass(), this);
+	if (Subsystem)
+	{
+		return Subsystem->LoadConfigObject(this->GetClass(), this);
+	}
+
+	return false;
 }
 
 bool UEditorConfigBase::SaveEditorConfig() const
 {
 	UEditorConfigSubsystem* Subsystem = GEditor->GetEditorSubsystem<UEditorConfigSubsystem>();
-	return Subsystem->SaveConfigObject(this->GetClass(), this);
+	if (Subsystem)
+	{
+		return Subsystem->SaveConfigObject(this->GetClass(), this);
+	}
+
+	return false;
 }
\ No newline at end of file
