UEditorConfigBase now won't load or save if GEditor is not initialized.

diff --git a/Engine/Source/Editor/EditorConfig/Private/EditorConfigBase.cpp b/Engine/Source/Editor/EditorConfig/Private/EditorConfigBase.cpp
--- a/Engine/Source/Editor/EditorConfig/Private/EditorConfigBase.cpp
+++ b/Engine/Source/Editor/EditorConfig/Private/EditorConfigBase.cpp
@@ -6,22 +6,26 @@
 
 bool UEditorConfigBase::LoadEditorConfig()
 {
-	UEditorConfigSubsystem* Subsystem = GEditor->GetEditorSubsystem<UEditorConfigSubsystem>();
-	if (Subsystem)
+	if (GEditor)
 	{
-		return Subsystem->LoadConfigObject(this->GetClass(), this);
+		UEditorConfigSubsystem* Subsystem = GEditor->GetEditorSubsystem<UEditorConfigSubsystem>();
+		if (Subsystem)
+		{
+			return Subsystem->LoadConfigObject(this->GetClass(), this);
+		}
 	}
-
 	return false;
 }
 
 bool UEditorConfigBase::SaveEditorConfig() const
 {
-	UEditorConfigSubsystem* Subsystem = GEditor->GetEditorSubsystem<UEditorConfigSubsystem>();
-	if (Subsystem)
+	if (GEditor)
 	{
-		return Subsystem->SaveConfigObject(this->GetClass(), this);
+		UEditorConfigSubsystem* Subsystem = GEditor->GetEditorSubsystem<UEditorConfigSubsystem>();
+		if (Subsystem)
+		{
+			return Subsystem->SaveConfigObject(this->GetClass(), this);
+		}
 	}
-
 	return false;
 }
\ No newline at end of file
