EditorConfig map handling fixed so that holes in the map won't cause an ensure - you must check IsValidIndex before accessing any value.

diff --git a/Engine/Source/Editor/EditorConfig/Private/EditorConfig.cpp b/Engine/Source/Editor/EditorConfig/Private/EditorConfig.cpp
--- a/Engine/Source/Editor/EditorConfig/Private/EditorConfig.cpp
+++ b/Engine/Source/Editor/EditorConfig/Private/EditorConfig.cpp
@@ -575,11 +575,14 @@ TSharedPtr<FJsonValue> FEditorConfig::WriteValue(const FProperty* Property, cons
 
 			for (int32 Idx = 0; Idx < MapHelper.Num(); ++Idx)
 			{
-				TSharedPtr<FJsonValue> JsonKey = WriteValue(KeyProperty, MapHelper.GetKeyPtr(Idx));
-				JsonKeysArray.Add(JsonKey);
+				if (MapHelper.IsValidIndex(Idx))
+				{
+					TSharedPtr<FJsonValue> JsonKey = WriteValue(KeyProperty, MapHelper.GetKeyPtr(Idx));
+					JsonKeysArray.Add(JsonKey);
 
-				TSharedPtr<FJsonValue> JsonValue = WriteValue(ValueProperty, MapHelper.GetValuePtr(Idx));
-				JsonValuesArray.Add(JsonValue);
+					TSharedPtr<FJsonValue> JsonValue = WriteValue(ValueProperty, MapHelper.GetValuePtr(Idx));
+					JsonValuesArray.Add(JsonValue);
+				}
 			}
 
 			// maps can either be stored as $key, $value pairs or, if the keys can be stringified, as a JSON object
