Details Favorites are now stored between sessions again.

The EPropertyFilter property was introduced, but FEditorMetadataOverrides was not updated to make use of it.

diff --git a/Engine/Source/Editor/EditorConfig/Private/EditorMetadataOverrides.cpp b/Engine/Source/Editor/EditorConfig/Private/EditorMetadataOverrides.cpp
--- a/Engine/Source/Editor/EditorConfig/Private/EditorMetadataOverrides.cpp
+++ b/Engine/Source/Editor/EditorConfig/Private/EditorMetadataOverrides.cpp
@@ -29,7 +29,7 @@ bool UEditorMetadataOverrides::LoadFromConfig(TSharedPtr<FEditorConfig> Config)
 		return false;
 	}
 
-	return SourceConfig->TryGetStruct(TEXT("Metadata"), LoadedMetadata);
+	return SourceConfig->TryGetStruct(TEXT("Metadata"), LoadedMetadata, FEditorConfig::EPropertyFilter::All);
 }
 
 void UEditorMetadataOverrides::Save()
@@ -39,7 +39,7 @@ void UEditorMetadataOverrides::Save()
 		return;
 	}
 
-	SourceConfig->SetStruct(TEXT("Metadata"), LoadedMetadata);
+	SourceConfig->SetStruct(TEXT("Metadata"), LoadedMetadata, FEditorConfig::EPropertyFilter::All);
 
 	UEditorConfigSubsystem* EditorConfigSubsystem = GEditor->GetEditorSubsystem<UEditorConfigSubsystem>();
 	if (EditorConfigSubsystem == nullptr)
