Expose AddSearchDirectory and LoadFromFile for EditorConfigs

diff --git a/Engine/Source/Editor/EditorConfig/Private/EditorConfigSubystem.cpp b/Engine/Source/Editor/EditorConfig/Private/EditorConfigSubystem.cpp
--- a/Engine/Source/Editor/EditorConfig/Private/EditorConfigSubystem.cpp
+++ b/Engine/Source/Editor/EditorConfig/Private/EditorConfigSubystem.cpp
@@ -14,9 +14,9 @@ UEditorConfigSubsystem::UEditorConfigSubsystem()
 
 void UEditorConfigSubsystem::Initialize(FSubsystemCollectionBase& Collection)
 {
-	AddSearchDirectory(FPaths::Combine(FPaths::EngineConfigDir(), TEXT("Editor"))); // Engine
-	AddSearchDirectory(FPaths::Combine(FPaths::ProjectConfigDir(), TEXT("Editor"))); // ProjectName
-	AddSearchDirectory(FPaths::Combine(FPlatformProcess::UserSettingsDir(), *FApp::GetEpicProductIdentifier(), TEXT("Editor"))); // AppData
+	AddSearchDirectory(ESearchDirectoryType::Engine, FPaths::Combine(FPaths::EngineConfigDir(), TEXT("Editor"))); // Engine
+	AddSearchDirectory(ESearchDirectoryType::Project, FPaths::Combine(FPaths::ProjectConfigDir(), TEXT("Editor"))); // ProjectName
+	AddSearchDirectory(ESearchDirectoryType::User, FPaths::Combine(FPlatformProcess::UserSettingsDir(), *FApp::GetEpicProductIdentifier(), TEXT("Editor"))); // AppData
 }
 
 void UEditorConfigSubsystem::Deinitialize()
@@ -110,9 +110,9 @@ bool UEditorConfigSubsystem::ReloadConfig(TSharedRef<FEditorConfig> Config)
 
 	TSharedPtr<FEditorConfig> Parent;
 
-	for (const FString& Dir : SearchDirectories)
+	for (const TPair<ESearchDirectoryType, FString>& Dir : SearchDirectories)
 	{
-		const FString FullPath = FPaths::Combine(Dir, ConfigName) + TEXT(".json");
+		const FString FullPath = FPaths::Combine(Dir.Value, ConfigName) + TEXT(".json");
 
 		// find an existing config or create one
 		const TSharedPtr<FEditorConfig>* Existing = LoadedConfigs.Find(FullPath);
@@ -147,7 +147,7 @@ TSharedRef<FEditorConfig> UEditorConfigSubsystem::FindOrLoadConfig(FStringView C
 	// look for the config in the final search directory and return if it's loaded
 	// this assumes that the hierarchy of configs is unchanged
 	// ie. given search directories [Foo, Bar], the existence of Bar/X.json is taken to mean that Foo/X.json has been loaded
-	const FString FinalPath = FPaths::Combine(SearchDirectories.Last(), ConfigNameString) + TEXT(".json");
+	const FString FinalPath = FPaths::Combine(SearchDirectories.Last().Value, ConfigNameString) + TEXT(".json");
 
 	const TSharedPtr<FEditorConfig>* FinalConfig = LoadedConfigs.Find(FinalPath);
 	if (FinalConfig != nullptr)
@@ -158,9 +158,9 @@ TSharedRef<FEditorConfig> UEditorConfigSubsystem::FindOrLoadConfig(FStringView C
 	// find or load all configs in all search directories with the given name
 	TSharedPtr<FEditorConfig> Parent;
 
-	for (const FString& Dir : SearchDirectories)
+	for (const TPair<ESearchDirectoryType, FString>& Dir : SearchDirectories)
 	{
-		const FString FullPath = FPaths::Combine(Dir, ConfigNameString) + TEXT(".json");
+		const FString FullPath = FPaths::Combine(Dir.Value, ConfigNameString) + TEXT(".json");
 
 		const TSharedPtr<FEditorConfig>* Existing = LoadedConfigs.Find(FullPath);
 		if (Existing != nullptr && Existing->IsValid())
@@ -272,7 +272,19 @@ void UEditorConfigSubsystem::OnSaveCompleted(TSharedPtr<FEditorConfig> Config)
 	}
 }
 
-void UEditorConfigSubsystem::AddSearchDirectory(FStringView SearchDir)
+void UEditorConfigSubsystem::AddSearchDirectory(ESearchDirectoryType Type, FStringView SearchDir)
 {
-	SearchDirectories.AddUnique(FString(SearchDir));
+	TPair<ESearchDirectoryType, FString> NewEntry(Type, SearchDir);
+	if (!SearchDirectories.Contains(NewEntry))
+	{
+		for (int32 i = 0; i < SearchDirectories.Num(); ++i)
+		{
+			if (SearchDirectories[i].Key > Type)
+			{
+				SearchDirectories.Insert(NewEntry, i);
+				return;
+			}
+		}
+		SearchDirectories.Add(NewEntry);
+	}
 }
\ No newline at end of file
diff --git a/Engine/Source/Editor/EditorConfig/Public/EditorConfigSubsystem.h b/Engine/Source/Editor/EditorConfig/Public/EditorConfigSubsystem.h
--- a/Engine/Source/Editor/EditorConfig/Public/EditorConfigSubsystem.h
+++ b/Engine/Source/Editor/EditorConfig/Public/EditorConfigSubsystem.h
@@ -68,10 +68,22 @@ public:
 	/** Force reload the given config and all its (current and potential) parents from disk. */
 	bool ReloadConfig(TSharedRef<FEditorConfig> Config);
 
+	enum class ESearchDirectoryType : uint8
+	{
+		Engine,
+		Project,
+		User
+	};
+
+	/**
+	 * Append a new config search directory to the given type.
+	 * Engine directories are searched first, then Project, then User.
+	 */
+	void AddSearchDirectory(ESearchDirectoryType Type, FStringView SearchDir);
+
 private:
 	void Initialize(FSubsystemCollectionBase& Collection) override;
 	void Deinitialize() override;
-	void AddSearchDirectory(FStringView SearchDir);
 	void OnSaveCompleted(TSharedPtr<FEditorConfig> Config);
 	void OnEditorConfigDirtied(const FEditorConfig& Config);
 
@@ -90,7 +102,7 @@ private:
 
 	FCriticalSection SaveLock;
 	TArray<FPendingSave> PendingSaves;
-	TArray<FString> SearchDirectories;
+	TArray<TPair<ESearchDirectoryType, FString>> SearchDirectories;
 	TMap<FString, TSharedPtr<FEditorConfig>> LoadedConfigs;
 };
 
