Log error when JSON deserializer fails to parse an EditorConfig file

diff --git a/Engine/Source/Editor/EditorConfig/Private/EditorConfigModule.cpp b/Engine/Source/Editor/EditorConfig/Private/EditorConfigModule.cpp
--- a/Engine/Source/Editor/EditorConfig/Private/EditorConfigModule.cpp
+++ b/Engine/Source/Editor/EditorConfig/Private/EditorConfigModule.cpp
@@ -3,4 +3,6 @@
 #include "EditorConfigModule.h"
 #include "Modules/ModuleManager.h"
 
+DEFINE_LOG_CATEGORY(LogEditorConfig);
+
 IMPLEMENT_MODULE(FEditorConfigModule, EditorConfig);
diff --git a/Engine/Source/Editor/EditorConfig/Private/JsonConfig.cpp b/Engine/Source/Editor/EditorConfig/Private/JsonConfig.cpp
--- a/Engine/Source/Editor/EditorConfig/Private/JsonConfig.cpp
+++ b/Engine/Source/Editor/EditorConfig/Private/JsonConfig.cpp
@@ -5,6 +5,7 @@
 #include "Misc/StringBuilder.h"
 #include "Serialization/JsonReader.h"
 #include "Serialization/JsonSerializer.h"
+#include "EditorConfigModule.h"
 
 DEFINE_LOG_CATEGORY_STATIC(LogJsonConfig, Log, All);
 
@@ -182,7 +183,12 @@ namespace UE
 			return false;
 		}
 
-		return LoadFromString(Contents);
+		if (!LoadFromString(Contents))
+		{
+			UE_LOG(LogEditorConfig, Error, TEXT("Failed to load JSON file into JsonConfig %s"), FilePath.GetData());
+			return false;
+		}
+		return true;
 	}
 
 	bool FJsonConfig::LoadFromString(FStringView Content)
@@ -190,6 +196,7 @@ namespace UE
 		TSharedRef<FJsonStringReader> JsonReader = FJsonStringReader::Create(Content.GetData());
 		if (!FJsonSerializer::Deserialize(JsonReader.Get(), OverrideObject))
 		{
+			UE_LOG(LogEditorConfig, Error, TEXT("Failed to deserialize JSON string"));
 			return false;
 		}
 
diff --git a/Engine/Source/Editor/EditorConfig/Public/EditorConfigModule.h b/Engine/Source/Editor/EditorConfig/Public/EditorConfigModule.h
--- a/Engine/Source/Editor/EditorConfig/Public/EditorConfigModule.h
+++ b/Engine/Source/Editor/EditorConfig/Public/EditorConfigModule.h
@@ -4,6 +4,8 @@
 
 #include "Modules/ModuleInterface.h"
 
+DECLARE_LOG_CATEGORY_EXTERN(LogEditorConfig, Log, All)
+
 class FEditorConfigModule : public IModuleInterface
 {
 };
