Fix non-unity error

diff --git a/Engine/Source/Editor/EditorConfig/Public/EditorConfigModule.h b/Engine/Source/Editor/EditorConfig/Public/EditorConfigModule.h
--- a/Engine/Source/Editor/EditorConfig/Public/EditorConfigModule.h
+++ b/Engine/Source/Editor/EditorConfig/Public/EditorConfigModule.h
@@ -3,8 +3,9 @@
 #pragma once
 
 #include "Modules/ModuleInterface.h"
+#include "Logging/LogMacros.h"
 
-DECLARE_LOG_CATEGORY_EXTERN(LogEditorConfig, Log, All)
+DECLARE_LOG_CATEGORY_EXTERN(LogEditorConfig, Log, All);
 
 class FEditorConfigModule : public IModuleInterface
 {
