Editor configs now correctly diff against the CDO when serializing.

Previously, I had assumed that FProperty::InitializeValue() would initialize a simple member (eg. a float) to the correct value, but this appears not to be the case, rather values are zero-initialized and the CDO stores the correct default values. In the case of structs, InitializeStruct() does in fact initialize correctly.

However, this means that we can't use GetMutableDefault() to get an instance of the UObject, since then we're changing what we're diffing against and serialize no values. Changed the usages of this pattern to use singletons instead.

diff --git a/Engine/Source/Editor/EditorConfig/Private/EditorConfig.cpp b/Engine/Source/Editor/EditorConfig/Private/EditorConfig.cpp
--- a/Engine/Source/Editor/EditorConfig/Private/EditorConfig.cpp
+++ b/Engine/Source/Editor/EditorConfig/Private/EditorConfig.cpp
@@ -127,13 +127,13 @@ void FEditorConfig::SetRootStruct(const UStruct* Struct, const void* Instance, E
 		return;
 	}
 	
-	TSharedPtr<FJsonObject> JsonObject = WriteStruct(Struct, Instance, Filter);
+	TSharedPtr<FJsonObject> JsonObject = WriteStruct(Struct, Instance, nullptr, Filter);
 	JsonConfig->SetRootObject(JsonObject);
 		
 	SetDirty();
 }
 
-void FEditorConfig::ReadStruct(TSharedPtr<FJsonObject> JsonObject, const UStruct* Struct, void* Instance, UObject* Owner, EPropertyFilter Filter)
+void FEditorConfig::ReadStruct(const TSharedPtr<FJsonObject>& JsonObject, const UStruct* Struct, void* Instance, UObject* Owner, EPropertyFilter Filter)
 {
 	FString TypeName;
 	JsonObject->TryGetStringField(TEXT("$type"), TypeName);
@@ -162,7 +162,7 @@ void FEditorConfig::ReadStruct(TSharedPtr<FJsonObject> JsonObject, const UStruct
 	}
 }
 
-void FEditorConfig::ReadUObject(TSharedPtr<FJsonObject> JsonObject, const UClass* Class, UObject* Instance, EPropertyFilter Filter)
+void FEditorConfig::ReadUObject(const TSharedPtr<FJsonObject>& JsonObject, const UClass* Class, UObject* Instance, EPropertyFilter Filter)
 {
 	FString TypeName;
 	JsonObject->TryGetStringField(TEXT("$type"), TypeName);
@@ -182,7 +182,7 @@ void FEditorConfig::ReadUObject(TSharedPtr<FJsonObject> JsonObject, const UClass
 		}
 
 		void* DataPtr = Property->ContainerPtrToValuePtr<void>(Instance);
-
+		
 		TSharedPtr<FJsonValue> Value = JsonObject->TryGetField(Property->GetName());
 		if (Value.IsValid())
 		{
@@ -191,12 +191,13 @@ void FEditorConfig::ReadUObject(TSharedPtr<FJsonObject> JsonObject, const UClass
 	}
 }
 
-void FEditorConfig::ReadValue(TSharedPtr<FJsonValue> JsonValue, const FProperty* Property, void* DataPtr, UObject* Owner)
+void FEditorConfig::ReadValue(const TSharedPtr<FJsonValue>& JsonValue, const FProperty* Property, void* DataPtr, UObject* Owner)
 {
 	if (const FStrProperty* StrProperty = CastField<FStrProperty>(Property))
 	{
 		FString* Value = (FString*) DataPtr;
 		JsonValue->TryGetString(*Value);
+		return;
 	}
 	else if (const FNameProperty* NameProperty = CastField<FNameProperty>(Property))
 	{
@@ -204,6 +205,7 @@ void FEditorConfig::ReadValue(TSharedPtr<FJsonValue> JsonValue, const FProperty*
 		JsonValue->TryGetString(TempValue);
 
 		*(FName*) DataPtr = *TempValue;
+		return;
 	}
 	else if (const FTextProperty* TextProperty = CastField<FTextProperty>(Property))
 	{
@@ -211,6 +213,7 @@ void FEditorConfig::ReadValue(TSharedPtr<FJsonValue> JsonValue, const FProperty*
 		JsonValue->TryGetString(TempValue);
 
 		*(FText*) DataPtr = FText::FromString(TempValue);
+		return;
 	}
 	else if (const FBoolProperty* BoolProperty = CastField<FBoolProperty>(Property))
 	{
@@ -219,56 +222,67 @@ void FEditorConfig::ReadValue(TSharedPtr<FJsonValue> JsonValue, const FProperty*
 		{
 			BoolProperty->SetPropertyValue(DataPtr, Value);
 		}
+		return;
 	}
 	else if (const FFloatProperty* FloatProperty = CastField<FFloatProperty>(Property))
 	{
 		float* Value = (float*) DataPtr;
 		JsonValue->TryGetNumber(*Value);
+		return;
 	}
 	else if (const FDoubleProperty* DoubleProperty = CastField<FDoubleProperty>(Property))
 	{
 		double* Value = (double*) DataPtr;
 		JsonValue->TryGetNumber(*Value);
-	} 
+		return;
+	}
 	else if (const FInt8Property* Int8Property = CastField<FInt8Property>(Property))
 	{
 		int8* Value = (int8*) DataPtr;
 		JsonValue->TryGetNumber(*Value);
+		return;
 	}
 	else if (const FInt16Property* Int16Property = CastField<FInt16Property>(Property))
 	{
 		int16* Value = (int16*) DataPtr;
 		JsonValue->TryGetNumber(*Value);
+		return;
 	}
 	else if (const FIntProperty* Int32Property = CastField<FIntProperty>(Property))
 	{
 		int32* Value = (int32*) DataPtr;
 		JsonValue->TryGetNumber(*Value);
+		return;
 	}
 	else if (const FInt64Property* Int64Property = CastField<FInt64Property>(Property))
 	{
 		int64* Value = (int64*) DataPtr;
 		JsonValue->TryGetNumber(*Value);
+		return;
 	}
 	else if (const FByteProperty* ByteProperty = CastField<FByteProperty>(Property))
 	{
 		uint8* Value = (uint8*) DataPtr;
 		JsonValue->TryGetNumber(*Value);
+		return;
 	}
 	else if (const FUInt16Property* Uint16Property = CastField<FUInt16Property>(Property))
 	{
 		uint16* Value = (uint16*) DataPtr;
 		JsonValue->TryGetNumber(*Value);
+		return;
 	}
 	else if (const FUInt32Property* Uint32Property = CastField<FUInt32Property>(Property))
 	{
 		uint32* Value = (uint32*) DataPtr;
 		JsonValue->TryGetNumber(*Value);
+		return;
 	}
 	else if (const FUInt64Property* Uint64Property = CastField<FUInt64Property>(Property))
 	{
 		uint64* Value = (uint64*) DataPtr;
 		JsonValue->TryGetNumber(*Value);
+		return;
 	}
 	else if (const FEnumProperty* EnumProperty = CastField<FEnumProperty>(Property))
 	{
@@ -287,6 +301,7 @@ void FEditorConfig::ReadValue(TSharedPtr<FJsonValue> JsonValue, const FProperty*
 				}
 			}
 		}
+		return;
 	}
 	else if (const FObjectPropertyBase* ObjectProperty = CastField<FObjectPropertyBase>(Property))
 	{
@@ -295,6 +310,7 @@ void FEditorConfig::ReadValue(TSharedPtr<FJsonValue> JsonValue, const FProperty*
 		{
 			Property->ImportText_Direct(*PathString, DataPtr, Owner, 0);
 		}
+		return;
 	}
 	else if (const FStructProperty* StructProperty = CastField<FStructProperty>(Property))
 	{
@@ -303,6 +319,7 @@ void FEditorConfig::ReadValue(TSharedPtr<FJsonValue> JsonValue, const FProperty*
 		{
 			ReadStruct(*ObjectJsonValue, StructProperty->Struct, DataPtr, Owner, EPropertyFilter::All);
 		}
+		return;
 	}
 	else if (const FArrayProperty* ArrayProperty = CastField<FArrayProperty>(Property))
 	{
@@ -320,6 +337,7 @@ void FEditorConfig::ReadValue(TSharedPtr<FJsonValue> JsonValue, const FProperty*
 				ReadValue(Value, InnerProperty, ArrayHelper.GetRawPtr(Idx), Owner);
 			}
 		}
+		return;
 	}
 	else if (const FSetProperty* SetProperty = CastField<FSetProperty>(Property))
 	{
@@ -339,13 +357,14 @@ void FEditorConfig::ReadValue(TSharedPtr<FJsonValue> JsonValue, const FProperty*
 				InnerProperty->InitializeValue(TempBuffer.GetData());
 
 				TSharedPtr<FJsonValue> Value = (*SetJsonValue)[Idx];
-				ReadValue(Value, InnerProperty, TempBuffer.GetData(), Owner); 
+				ReadValue(Value, InnerProperty, TempBuffer.GetData(), Owner);
 
 				SetHelper.AddElement(TempBuffer.GetData());
 
 				InnerProperty->DestroyValue(TempBuffer.GetData());
 			}
 		}
+		return;
 	}
 	else if (const FMapProperty* MapProperty = CastField<FMapProperty>(Property))
 	{
@@ -374,14 +393,14 @@ void FEditorConfig::ReadValue(TSharedPtr<FJsonValue> JsonValue, const FProperty*
 				KeyProperty->ImportText_Direct(*JsonPair.Key, TempKey.GetData(), Owner, 0);
 
 				ValueProperty->InitializeValue(TempValue.GetData());
-				ReadValue(JsonPair.Value, ValueProperty, TempValue.GetData(), Owner); 
+				ReadValue(JsonPair.Value, ValueProperty, TempValue.GetData(), Owner);
 
 				MapHelper.AddPair(TempKey.GetData(), TempValue.GetData());
 
 				KeyProperty->DestroyValue(TempKey.GetData());
 				ValueProperty->DestroyValue(TempValue.GetData());
 			}
-		
+
 			return;
 		}
 
@@ -390,7 +409,7 @@ void FEditorConfig::ReadValue(TSharedPtr<FJsonValue> JsonValue, const FProperty*
 		if (JsonValue->TryGetArray(JsonArrayPtr))
 		{
 			MapHelper.EmptyValues(JsonArrayPtr->Num());
-			
+
 			// temporary buffers to read elements into
 			TArray<uint8> TempKey;
 			TempKey.AddUninitialized(KeyProperty->ElementSize);
@@ -409,10 +428,10 @@ void FEditorConfig::ReadValue(TSharedPtr<FJsonValue> JsonValue, const FProperty*
 					if (JsonKeyField.IsValid() && JsonValueField.IsValid())
 					{
 						KeyProperty->InitializeValue(TempKey.GetData());
-						ReadValue(JsonKeyField, KeyProperty, TempKey.GetData(), Owner); 
+						ReadValue(JsonKeyField, KeyProperty, TempKey.GetData(), Owner);
 
 						ValueProperty->InitializeValue(TempValue.GetData());
-						ReadValue(JsonValueField, ValueProperty, TempValue.GetData(), Owner); 
+						ReadValue(JsonValueField, ValueProperty, TempValue.GetData(), Owner);
 
 						MapHelper.AddPair(TempKey.GetData(), TempValue.GetData());
 
@@ -425,11 +444,142 @@ void FEditorConfig::ReadValue(TSharedPtr<FJsonValue> JsonValue, const FProperty*
 			return;
 		}
 	}
+
+	ensureAlwaysMsgf(false, TEXT("Property type is unsupported: %s, type: %s"), *Property->GetPathName(), *Property->GetClass()->GetName());
+}
+
+TSharedPtr<FJsonValue> FEditorConfig::WriteArray(const FArrayProperty* ArrayProperty, const void* DataPtr)
+{
+	FProperty* InnerProperty = ArrayProperty->Inner;
+	FScriptArrayHelper ArrayHelper(ArrayProperty, DataPtr);
+
+	TArray<TSharedPtr<FJsonValue>> JsonValuesArray;
+	JsonValuesArray.Reserve(ArrayHelper.Num());
+
+	for (int32 Idx = 0; Idx < ArrayHelper.Num(); ++Idx)
+	{
+		if (ArrayHelper.IsValidIndex(Idx))
+		{
+			TSharedPtr<FJsonValue> ElementValue = WriteValue(InnerProperty, ArrayHelper.GetRawPtr(Idx), nullptr);
+			check(ElementValue.IsValid());
+			JsonValuesArray.Add(ElementValue);
+		}
+	}
+
+	return MakeShared<FJsonValueArray>(JsonValuesArray);
+}
+
+TSharedPtr<FJsonValue> FEditorConfig::WriteSet(const FSetProperty* SetProperty, const void* DataPtr)
+{
+	FProperty* InnerProperty = SetProperty->ElementProp;
+	FScriptSetHelper SetHelper(SetProperty, DataPtr);
+
+	TArray<TSharedPtr<FJsonValue>> JsonValuesArray;
+	JsonValuesArray.Reserve(SetHelper.Num());
+
+	for (int32 Idx = 0; Idx < SetHelper.Num(); ++Idx)
+	{
+		if (SetHelper.IsValidIndex(Idx))
+		{
+			TSharedPtr<FJsonValue> ElementValue = WriteValue(InnerProperty, SetHelper.GetElementPtr(Idx), nullptr);
+			check(ElementValue.IsValid());
+			JsonValuesArray.Add(ElementValue);
+		}
+	}
+
+	return MakeShared<FJsonValueArray>(JsonValuesArray);
+}
+
+TSharedPtr<FJsonValue> FEditorConfig::WriteMap(const FMapProperty* MapProperty, const void* DataPtr)
+{
+	TSharedPtr<FJsonValue> ResultValue;
+
+	FProperty* KeyProperty = MapProperty->KeyProp;
+	FProperty* ValueProperty = MapProperty->ValueProp;
+
+	FScriptMapHelper MapHelper(MapProperty, DataPtr);
+
+	if (MapHelper.Num() == 0)
+	{
+		ResultValue = MakeShared<FJsonValueObject>(MakeShared<FJsonObject>());
+	}
+	else
+	{
+		TArray<TSharedPtr<FJsonValue>> JsonKeysArray;
+		JsonKeysArray.Reserve(MapHelper.Num());
+
+		TArray<TSharedPtr<FJsonValue>> JsonValuesArray;
+		JsonValuesArray.Reserve(MapHelper.Num());
+
+		for (int32 Idx = 0; Idx < MapHelper.Num(); ++Idx)
+		{
+			if (MapHelper.IsValidIndex(Idx))
+			{
+				TSharedPtr<FJsonValue> JsonKey = WriteValue(KeyProperty, MapHelper.GetKeyPtr(Idx), nullptr);
+				check(JsonKey.IsValid());
+				JsonKeysArray.Add(JsonKey);
+
+				TSharedPtr<FJsonValue> JsonValue = WriteValue(ValueProperty, MapHelper.GetValuePtr(Idx), nullptr);
+				check(JsonValue.IsValid());
+				JsonValuesArray.Add(JsonValue);
+			}
+		}
+
+		// maps can either be stored as $key, $value pairs or, if the keys can be stringified, as a JSON object
+		// check Filter we should use based on the first element
+		EJson KeyType = JsonKeysArray[0]->Type;
+		if (KeyType == EJson::Object)
+		{
+			TArray<TSharedPtr<FJsonValue>> ResultArray;
+			ResultArray.Reserve(MapHelper.Num());
+
+			for (int32 Idx = 0; Idx < MapHelper.Num(); ++Idx)
+			{
+				if (MapHelper.IsValidIndex(Idx))
+				{
+					TSharedPtr<FJsonObject> ElementObject = MakeShared<FJsonObject>();
+					ElementObject->SetField(TEXT("$key"), JsonKeysArray[Idx]);
+					ElementObject->SetField(TEXT("$value"), JsonValuesArray[Idx]);
+
+					ResultArray.Add(MakeShared<FJsonValueObject>(ElementObject));
+				}
+			}
+
+			ResultValue = MakeShared<FJsonValueArray>(ResultArray);
+		}
+		else if (KeyType == EJson::Boolean ||
+			KeyType == EJson::Number ||
+			KeyType == EJson::String)
+		{
+			TSharedPtr<FJsonObject> ResultObject = MakeShared<FJsonObject>();
+
+			for (int32 Idx = 0; Idx < MapHelper.Num(); ++Idx)
+			{
+				if (MapHelper.IsValidIndex(Idx))
+				{
+					FString KeyString;
+					check(JsonKeysArray[Idx]->TryGetString(KeyString));
+
+					ResultObject->SetField(KeyString, JsonValuesArray[Idx]);
+				}
+			}
+
+			ResultValue = MakeShared<FJsonValueObject>(ResultObject);
+		}
+
+		ensureMsgf(ResultValue.IsValid(), TEXT("Map key type is invalid."));
+	}
+
+	return ResultValue;
 }
 
-TSharedPtr<FJsonValue> FEditorConfig::WriteValue(const FProperty* Property, const void* DataPtr)
+TSharedPtr<FJsonValue> FEditorConfig::WriteValue(const FProperty* Property, const void* DataPtr, const void* DefaultPtr)
 {
 	TSharedPtr<FJsonValue> ResultValue;
+	if (DefaultPtr != nullptr && Property->Identical(DataPtr, DefaultPtr))
+	{
+		return ResultValue;
+	}
 
 	if (const FStrProperty* StrProperty = CastField<FStrProperty>(Property))
 	{
@@ -517,144 +667,45 @@ TSharedPtr<FJsonValue> FEditorConfig::WriteValue(const FProperty* Property, cons
 	}
 	else if (const FStructProperty* StructProperty = CastField<FStructProperty>(Property))
 	{
-		ResultValue = MakeShared<FJsonValueObject>(WriteStruct(StructProperty->Struct, DataPtr, EPropertyFilter::All));
+		ResultValue = MakeShared<FJsonValueObject>(WriteStruct(StructProperty->Struct, DataPtr, DefaultPtr, EPropertyFilter::All));
 	}
 	else if (const FArrayProperty* ArrayProperty = CastField<FArrayProperty>(Property))
 	{
-		FProperty* InnerProperty = ArrayProperty->Inner;
-		FScriptArrayHelper ArrayHelper(ArrayProperty, DataPtr);
-
-		TArray<TSharedPtr<FJsonValue>> JsonValuesArray;
-		JsonValuesArray.Reserve(ArrayHelper.Num());
-
-		for (int32 Idx = 0; Idx < ArrayHelper.Num(); ++Idx)
-		{
-			TSharedPtr<FJsonValue> ElementValue = WriteValue(InnerProperty, ArrayHelper.GetRawPtr(Idx));
-			JsonValuesArray.Add(ElementValue);
-		}
-
-		ResultValue = MakeShared<FJsonValueArray>(JsonValuesArray);
+		ResultValue = WriteArray(ArrayProperty, DataPtr);
 	}
 	else if (const FSetProperty* SetProperty = CastField<FSetProperty>(Property))
 	{
-		FProperty* InnerProperty = SetProperty->ElementProp;
-		FScriptSetHelper SetHelper(SetProperty, DataPtr);
-
-		TArray<TSharedPtr<FJsonValue>> JsonValuesArray;
-		JsonValuesArray.Reserve(SetHelper.Num());
-
-		for (int32 Idx = 0; Idx < SetHelper.Num(); ++Idx)
-		{
-			if (SetHelper.IsValidIndex(Idx))
-			{
-				TSharedPtr<FJsonValue> ElementValue = WriteValue(InnerProperty, SetHelper.GetElementPtr(Idx));
-				JsonValuesArray.Add(ElementValue);
-			}
-		}
-
-		ResultValue = MakeShared<FJsonValueArray>(JsonValuesArray);
+		ResultValue = WriteSet(SetProperty, DataPtr);
 	}
 	else if (const FMapProperty* MapProperty = CastField<FMapProperty>(Property))
 	{
-		FProperty* KeyProperty = MapProperty->KeyProp;
-		FProperty* ValueProperty = MapProperty->ValueProp;
-
-		FScriptMapHelper MapHelper(MapProperty, DataPtr);
-
-		if (MapHelper.Num() == 0)
-		{
-			ResultValue = MakeShared<FJsonValueObject>(MakeShared<FJsonObject>());
-		}
-		else
-		{
-			TArray<TSharedPtr<FJsonValue>> JsonKeysArray;
-			JsonKeysArray.Reserve(MapHelper.Num());
-
-			TArray<TSharedPtr<FJsonValue>> JsonValuesArray;
-			JsonValuesArray.Reserve(MapHelper.Num());
-
-			for (int32 Idx = 0; Idx < MapHelper.Num(); ++Idx)
-			{
-				if (MapHelper.IsValidIndex(Idx))
-				{
-					TSharedPtr<FJsonValue> JsonKey = WriteValue(KeyProperty, MapHelper.GetKeyPtr(Idx));
-					JsonKeysArray.Add(JsonKey);
-
-					TSharedPtr<FJsonValue> JsonValue = WriteValue(ValueProperty, MapHelper.GetValuePtr(Idx));
-					JsonValuesArray.Add(JsonValue);
-				}
-			}
-
-			// maps can either be stored as $key, $value pairs or, if the keys can be stringified, as a JSON object
-			// check Filter we should use based on the first element
-			EJson KeyType = JsonKeysArray[0]->Type;
-			if (KeyType == EJson::Object)
-			{
-				TArray<TSharedPtr<FJsonValue>> ResultArray;
-				ResultArray.Reserve(MapHelper.Num());
-
-				for (int32 Idx = 0; Idx < MapHelper.Num(); ++Idx)
-				{
-					TSharedPtr<FJsonObject> ElementObject = MakeShared<FJsonObject>();
-					ElementObject->SetField(TEXT("$key"), JsonKeysArray[Idx]);
-					ElementObject->SetField(TEXT("$value"), JsonValuesArray[Idx]);
-
-					ResultArray.Add(MakeShared<FJsonValueObject>(ElementObject));
-				}
-
-				ResultValue = MakeShared<FJsonValueArray>(ResultArray);
-			}
-			else if (KeyType == EJson::Boolean || 
-				KeyType == EJson::Number || 
-				KeyType == EJson::String)
-			{
-				TSharedPtr<FJsonObject> ResultObject = MakeShared<FJsonObject>();
-
-				for (int32 Idx = 0; Idx < MapHelper.Num(); ++Idx)
-				{
-					FString KeyString;
-					check(JsonKeysArray[Idx]->TryGetString(KeyString));
-
-					ResultObject->SetField(KeyString, JsonValuesArray[Idx]);
-				}
-				
-				ResultValue = MakeShared<FJsonValueObject>(ResultObject);
-			}
-
-			ensureMsgf(ResultValue.IsValid(), TEXT("Map key type is invalid."));
-		}
+		ResultValue = WriteMap(MapProperty, DataPtr);
 	}
 
-	ensureMsgf(ResultValue.IsValid(), TEXT("Property type is unsupported."));
+	ensureAlwaysMsgf(ResultValue.IsValid(), TEXT("Property type is unsupported: %s, type: %s"), *Property->GetPathName(), *Property->GetClass()->GetName());
 	return ResultValue;
 }
 
-bool FEditorConfig::IsDefault(const FProperty* Property, TSharedPtr<FJsonValue> JsonValue, const void* NativeValue)
+TSharedPtr<FJsonObject> FEditorConfig::WriteStruct(const UStruct* Struct, const void* Instance, const void* Defaults, EPropertyFilter Filter) 
 {
-	if (JsonValue->Type == EJson::Array)
-	{
-		return JsonValue->AsArray().Num() == 0;
-	}
-	else if (JsonValue->Type == EJson::Object)
-	{
-		return JsonValue->AsObject()->Values.Num() == 0;
-	}
-	else
+	TSharedPtr<FJsonObject> JsonObject = MakeShared<FJsonObject>();
+	JsonObject->SetStringField(TEXT("$type"), Struct->GetName());
+	
+	// Default initialize a struct here if we weren't passed one. 
+	// This is necessary because structs can contain default initializations, eg.:
+	// struct Foo { float Bar = 5.0f; }
+	// The Bar FProperty does not store this value itself, only the struct does.
+	TArray<uint8> TempDefaults;
+
+	if (Defaults == nullptr)
 	{
-		// have the property initialize some temp storage, then check against that
-		uint8 Temp[256];
-		Property->InitializeValue(Temp);
+		TempDefaults.AddZeroed(Struct->GetStructureSize());
+		Struct->InitializeStruct(TempDefaults.GetData());
 
-		return Property->Identical(NativeValue, Temp);
+		Defaults = TempDefaults.GetData();
 	}
 
-	return false;
-}
-
-TSharedPtr<FJsonObject> FEditorConfig::WriteStruct(const UStruct* Struct, const void* Instance, EPropertyFilter Filter) 
-{
-	TSharedPtr<FJsonObject> JsonObject = MakeShared<FJsonObject>();
-	JsonObject->SetStringField(TEXT("$type"), Struct->GetName());
+	bool bAnyWritten = false;
 
 	for (TFieldIterator<FProperty> It(Struct); It; ++It)
 	{
@@ -665,15 +716,20 @@ TSharedPtr<FJsonObject> FEditorConfig::WriteStruct(const UStruct* Struct, const
 			continue;
 		}
 
-		const void* ValuePtr = Property->ContainerPtrToValuePtr<void>(Instance); 
+		bAnyWritten = true;
+
+		const void* ValuePtr = Property->ContainerPtrToValuePtr<void>(Instance);
+		const void* PropertyDefaultPtr = Property->ContainerPtrToValuePtr<void>(Defaults);
 
-		TSharedPtr<FJsonValue> PropertyValue = WriteValue(Property, ValuePtr);
-		if (!IsDefault(Property, PropertyValue, ValuePtr))
+		TSharedPtr<FJsonValue> PropertyValue = WriteValue(Property, ValuePtr, PropertyDefaultPtr);
+		if (PropertyValue.IsValid())
 		{
 			JsonObject->SetField(Property->GetName(), PropertyValue);
 		}
 	}
 
+	ensureAlwaysMsgf(bAnyWritten, TEXT("Struct type has no properties to serialize: %s"), *Struct->GetName());
+
 	return JsonObject;
 }
 
@@ -685,6 +741,10 @@ TSharedPtr<FJsonObject> FEditorConfig::WriteUObject(const UClass* Class, const U
 	TSharedPtr<FJsonObject> JsonObject = MakeShared<FJsonObject>();
 	JsonObject->SetStringField(TEXT("$type"), Class->GetName());
 
+	const UObject* Defaults = Class->GetDefaultObject();
+
+	bool bAnyWritten = false;
+
 	for (TFieldIterator<FProperty> It(Class); It; ++It)
 	{
 		const FProperty* Property = *It;
@@ -694,15 +754,20 @@ TSharedPtr<FJsonObject> FEditorConfig::WriteUObject(const UClass* Class, const U
 			continue;
 		}
 
-		const void* ValuePtr = Property->ContainerPtrToValuePtr<void>(Instance); 
+		bAnyWritten = true;
 
-		TSharedPtr<FJsonValue> PropertyValue = WriteValue(Property, ValuePtr);
-		if (!IsDefault(Property, PropertyValue, ValuePtr))
+		const void* ValuePtr = Property->ContainerPtrToValuePtr<void>(Instance);
+		const void* PropertyDefaultPtr = Property->ContainerPtrToValuePtr<void>(Defaults);
+
+		TSharedPtr<FJsonValue> PropertyValue = WriteValue(Property, ValuePtr, PropertyDefaultPtr);
+		if (PropertyValue.IsValid())
 		{
 			JsonObject->SetField(Property->GetName(), PropertyValue);
 		}
 	}
 
+	ensureAlwaysMsgf(bAnyWritten, TEXT("UObject type has no properties to serialize: %s"), *Class->GetName());
+
 	return JsonObject;
 }
 
diff --git a/Engine/Source/Editor/EditorConfig/Private/Tests/EditorConfigTests.h b/Engine/Source/Editor/EditorConfig/Private/Tests/EditorConfigTests.h
--- a/Engine/Source/Editor/EditorConfig/Private/Tests/EditorConfigTests.h
+++ b/Engine/Source/Editor/EditorConfig/Private/Tests/EditorConfigTests.h
@@ -8,28 +8,31 @@
 #include "EditorConfigTests.generated.h"
 
 USTRUCT()
-struct FEditorConfigTestStruct
+struct FEditorConfigTestSimpleStruct
 {
 	GENERATED_BODY()
 
 	UPROPERTY()
-	bool BoolProperty = false;
+	bool Bool = true;
 
 	UPROPERTY()
-	int32 IntProperty = 0;
+	int32 Int = 5;
 
-	UPROPERTY()
-	FString StringProperty;
+	UPROPERTY(meta=(EditorConfig))
+	FString String;
 
 	UPROPERTY()
-	float FloatProperty = 0.0f;
+	float Float = 5.0f;
 
-	UPROPERTY()
-	TArray<FString> ArrayProperty;
+	UPROPERTY(meta=(EditorConfig))
+	TArray<FString> Array;
+
+	bool operator==(const FEditorConfigTestSimpleStruct& Other) const;
+	bool operator!=(const FEditorConfigTestSimpleStruct& Other) const;
 };
 
 USTRUCT()
-struct FEditorConfigTestKeyStruct
+struct FEditorConfigTestKey
 {
 	GENERATED_BODY()
 	
@@ -37,28 +40,28 @@ struct FEditorConfigTestKeyStruct
 	FString Name;
 	
 	UPROPERTY()
-	double Number = 0.0;
+	double Number = 5.0;
 
-	bool operator==(const FEditorConfigTestKeyStruct& Other) const;
-	bool operator!=(const FEditorConfigTestKeyStruct& Other) const;
+	bool operator==(const FEditorConfigTestKey& Other) const;
+	bool operator!=(const FEditorConfigTestKey& Other) const;
 };
 
-FORCEINLINE uint32 GetTypeHash(const FEditorConfigTestKeyStruct& Struct)
+FORCEINLINE uint32 GetTypeHash(const FEditorConfigTestKey& Struct)
 {
 	return HashCombine(GetTypeHash(Struct.Name), GetTypeHash(Struct.Number));
 }
 
 USTRUCT()
-struct FEditorConfigTestComplexArrayStruct
+struct FEditorConfigTestComplexArray
 {
 	GENERATED_BODY()
 	
 	UPROPERTY()
-	TArray<FEditorConfigTestKeyStruct> Array;
+	TArray<FEditorConfigTestKey> Array;
 };
 
 USTRUCT()
-struct FEditorConfigTestSimpleMapStruct
+struct FEditorConfigTestSimpleMap
 {
 	GENERATED_BODY()
 
@@ -67,16 +70,16 @@ struct FEditorConfigTestSimpleMapStruct
 };
 
 USTRUCT()
-struct FEditorConfigTestComplexMapStruct
+struct FEditorConfigTestComplexMap
 {
 	GENERATED_BODY()
 	
 	UPROPERTY()
-	TMap<FEditorConfigTestKeyStruct, FEditorConfigTestKeyStruct> Map;
+	TMap<FEditorConfigTestKey, FEditorConfigTestKey> Map;
 };
 
 USTRUCT()
-struct FEditorConfigTestSimpleSetStruct
+struct FEditorConfigTestSimpleSet
 {
 	GENERATED_BODY()
 	
@@ -85,12 +88,12 @@ struct FEditorConfigTestSimpleSetStruct
 };
 
 USTRUCT()
-struct FEditorConfigTestComplexSetStruct
+struct FEditorConfigTestComplexSet
 {
 	GENERATED_BODY()
 	
 	UPROPERTY()
-	TSet<FEditorConfigTestKeyStruct> Set;
+	TSet<FEditorConfigTestKey> Set;
 };
 
 UCLASS()
@@ -98,15 +101,14 @@ class UEditorConfigTestObject : public UObject
 {
 	GENERATED_BODY()
 
-	UPROPERTY()
-	TObjectPtr<UObject> Object;
+public:
 
 	UPROPERTY()
-	FSoftObjectPath SoftObjectPath;
+	TObjectPtr<UObject> Object;
 
-	UPROPERTY()
-	FEditorConfigTestStruct Struct;
+	UPROPERTY(meta=(EditorConfig))
+	FEditorConfigTestSimpleStruct Struct;
 
-	UPROPERTY()
-	int32 Number;
+	UPROPERTY(meta=(EditorConfig))
+	int32 Number = 5;
 };
diff --git a/Engine/Source/Editor/EditorConfig/Private/Tests/JsonConfigTests.cpp b/Engine/Source/Editor/EditorConfig/Private/Tests/JsonConfigTests.cpp
--- a/Engine/Source/Editor/EditorConfig/Private/Tests/JsonConfigTests.cpp
+++ b/Engine/Source/Editor/EditorConfig/Private/Tests/JsonConfigTests.cpp
@@ -1,6 +1,5 @@
 // Copyright Epic Games, Inc. All Rights Reserved.
 
-
 #if WITH_DEV_AUTOMATION_TESTS 
 
 #include "EditorConfigTestHelpers.h"
@@ -1208,89 +1207,4 @@ R"_JSON({
 	return true;
 }
 
-
-/*IMPLEMENT_SIMPLE_AUTOMATION_TEST(FJsonConfigTests_SetStringInObjectInArray, "JsonConfig.Set.Object.InArray", EAutomationTestFlags::EditorContext | EAutomationTestFlags::EngineFilter)
-
-bool FJsonConfigTests_SetObjectInArray::RunTest(const FString& Parameters)
-{
-	const FString ParentContents =
-R"_JSON({
-	"Array": 
-	[
-		{
-			"String": "Foo",
-			"Object":
-			{
-				"Bool": false,
-				"Number": 1
-			}
-		}
-	]
-})_JSON";
-
-	TSharedRef<FJsonConfig> ParentConfig = MakeShared<FJsonConfig>();
-	if (!ParentConfig->LoadFromString(ParentContents))
-	{
-		return false;
-	}
-
-	FJsonConfig ChildConfig;
-	ChildConfig.SetParent(ParentConfig);
-
-	TSharedPtr<FJsonObject> Object;
-	if (!ChildConfig.TryGetJsonObject(TEXT("Array[0]"), Object))
-	{
-		return false;
-	}
-
-	TSharedPtr<FJsonValue> StringValue, ObjectValue;
-
-	StringValue = Object->Values["String"];
-	TestEqual(TEXT("Before Set"), StringValue->AsString(), TEXT("Foo"));
-
-	ObjectValue = Object->Values["Object"];
-
-	bool OutBool;
-	if (!ObjectValue->AsObject()->TryGetBoolField(TEXT("Bool"), OutBool))
-	{
-		return false;
-	}
-
-	TestEqual(TEXT("Before Set"), OutBool, false);
-
-	int32 OutNumber;
-	if (!ObjectValue->AsObject()->TryGetNumberField(TEXT("Number"), OutNumber))
-	{
-		return false;
-	}
-
-	TestEqual(TEXT("Before Set"), OutNumber, 1);
-
-	//ChildConfig.SetString(TEXT("Array[0].String"), TEXT("Bar"));
-
-	if (!ChildConfig.TryGetJsonObject(TEXT("Array[0]"), Object))
-	{
-		return false;
-	}
-
-	StringValue = Object->Values["String"];
-	TestEqual(TEXT("After String Set"), StringValue->AsString(), TEXT("Bar"));
-
-	const FString ExpectedOverrideString1 = 
-R"_JSON({
-	"Array":
-	[
-		{
-			
-		}
-	]
-})_JSON";
-
-	FString OverrideString;
-	ChildConfig.SaveToString(OverrideString);
-	TestTrue(TEXT("Serialized 1"), FEditorConfigTestHelpers::AreJsonStringsEquivalent(OverrideString, ExpectedOverrideString1));
-
-	return true;
-}*/
-
 #endif
\ No newline at end of file
diff --git a/Engine/Source/Editor/EditorConfig/Public/EditorConfig.h b/Engine/Source/Editor/EditorConfig/Public/EditorConfig.h
--- a/Engine/Source/Editor/EditorConfig/Public/EditorConfig.h
+++ b/Engine/Source/Editor/EditorConfig/Public/EditorConfig.h
@@ -5,7 +5,9 @@
 #include "JsonConfig.h"
 #include "Concepts/EqualityComparable.h"
 #include "Dom/JsonObject.h"
+#include "Templates/IsPointer.h"
 #include "Templates/Models.h"
+#include "Templates/UnrealTemplate.h"
 
 class EDITORCONFIG_API FEditorConfig
 {
@@ -38,6 +40,8 @@ public:
 	bool TryGetRootStruct(T& OutValue, EPropertyFilter Filter = EPropertyFilter::MetadataOnly) const;
 	template <typename T>
 	bool TryGetRootUObject(T& OutValue, EPropertyFilter Filter = EPropertyFilter::MetadataOnly) const;
+	template <typename T>
+	bool TryGetRootUObject(T* OutValue, EPropertyFilter Filter = EPropertyFilter::MetadataOnly) const;
 
 	bool TryGetRootStruct(const UStruct* Class, void* OutValue, EPropertyFilter Filter = EPropertyFilter::MetadataOnly) const;
 	bool TryGetRootUObject(const UClass* Class, UObject* OutValue, EPropertyFilter Filter = EPropertyFilter::MetadataOnly) const;
@@ -66,15 +70,16 @@ private:
 
 	friend class UEditorConfigSubsystem; // for access to LoadFromFile and SaveToFile
 
-	static void ReadUObject(TSharedPtr<FJsonObject> JsonObject, const UClass* Class, UObject* Instance, EPropertyFilter Filter);
-	static void ReadStruct(TSharedPtr<FJsonObject> JsonObject, const UStruct* Struct, void* Instance, UObject* Owner, EPropertyFilter Filter);
-	static void ReadValue(TSharedPtr<FJsonValue> JsonValue, const FProperty* Property, void* DataPtr, UObject* Owner);
+	static void ReadUObject(const TSharedPtr<FJsonObject>& JsonObject, const UClass* Class, UObject* Instance, EPropertyFilter Filter);
+	static void ReadStruct(const TSharedPtr<FJsonObject>& JsonObject, const UStruct* Struct, void* Instance, UObject* Owner, EPropertyFilter Filter);
+	static void ReadValue(const TSharedPtr<FJsonValue>& JsonValue, const FProperty* Property, void* DataPtr, UObject* Owner);
 	
-	static TSharedPtr<FJsonObject> WriteStruct(const UStruct* Struct, const void* Instance, EPropertyFilter Filter);
+	static TSharedPtr<FJsonObject> WriteStruct(const UStruct* Struct, const void* Instance, const void* Defaults, EPropertyFilter Filter);
 	static TSharedPtr<FJsonObject> WriteUObject(const UClass* Class, const UObject* Instance, EPropertyFilter Filter);
-	static TSharedPtr<FJsonValue> WriteValue(const FProperty* Property, const void* DataPtr);
-
-	static bool IsDefault(const FProperty* Property, TSharedPtr<FJsonValue> JsonValue, const void* NativeValue);
+	static TSharedPtr<FJsonValue> WriteArray(const FArrayProperty* ArrayProperty, const void* DataPtr);
+	static TSharedPtr<FJsonValue> WriteSet(const FSetProperty* Property, const void* DataPtr);
+	static TSharedPtr<FJsonValue> WriteMap(const FMapProperty* Property, const void* DataPtr);
+	static TSharedPtr<FJsonValue> WriteValue(const FProperty* Property, const void* DataPtr, const void* Defaults);
 
 	bool SaveToFile(FStringView FilePath) const;
 
@@ -156,6 +161,15 @@ bool FEditorConfig::TryGetRootUObject(T& OutValue, EPropertyFilter Filter) const
 	return TryGetRootUObject(T::StaticClass(), &OutValue, Filter);
 }
 
+template <typename T>
+bool FEditorConfig::TryGetRootUObject(T* OutValue, EPropertyFilter Filter) const
+{
+	static_assert(TIsDerivedFrom<T, UObject>::Value, "Type is not derived from UObject.");
+	checkf(OutValue != nullptr, TEXT("Output value was null."));
+
+	return TryGetRootUObject(T::StaticClass(), OutValue, Filter);
+}
+
 template <typename T>
 void FEditorConfig::SetStruct(FStringView Key, const T& InValue, EPropertyFilter Filter)
 {
@@ -164,7 +178,7 @@ void FEditorConfig::SetStruct(FStringView Key, const T& InValue, EPropertyFilter
 		return;
 	}
 
-	TSharedPtr<FJsonObject> JsonObject = WriteStruct(T::StaticStruct(), &InValue, Filter);
+	TSharedPtr<FJsonObject> JsonObject = WriteStruct(T::StaticStruct(), &InValue, nullptr, Filter);
 	JsonConfig->SetJsonObject(UE::FJsonPath(Key), JsonObject);
 
 	SetDirty();
@@ -195,7 +209,17 @@ void FEditorConfig::SetRootStruct(const T& InValue, EPropertyFilter Filter)
 template <typename T>
 void FEditorConfig::SetRootUObject(const T& InValue, EPropertyFilter Filter)
 {
-	static_assert(TIsDerivedFrom<T, UObject>::Value, "Type is not derived from UObject.");
+	if constexpr (TIsPointer<T>::Value)
+	{
+		static_assert(TIsDerivedFrom<typename TRemovePointer<T>::Type, UObject>::Value, "Type is not derived from UObject.");
+		checkf(InValue != nullptr, TEXT("Object value was null."));
 
-	SetRootUObject(T::StaticClass(), &InValue, Filter);
-}
+		SetRootUObject(TRemovePointer<T>::Type::StaticClass(), InValue, Filter);
+	}
+	else
+	{
+		static_assert(TIsDerivedFrom<typename TRemovePointer<T>::Type, UObject>::Value, "Type is not derived from UObject.");
+
+		SetRootUObject(T::StaticClass(), &InValue, Filter);
+	}
+}
\ No newline at end of file
diff --git a/Engine/Source/Editor/PropertyEditor/Private/DetailsViewConfig.cpp b/Engine/Source/Editor/PropertyEditor/Private/DetailsViewConfig.cpp
new file mode 100644
--- /dev/null
+++ b/Engine/Source/Editor/PropertyEditor/Private/DetailsViewConfig.cpp
@@ -0,0 +1,19 @@
+// Copyright Epic Games, Inc. All Rights Reserved.
+
+#include "DetailsViewConfig.h"
+
+TObjectPtr<UDetailsConfig> UDetailsConfig::Instance = nullptr;
+
+void UDetailsConfig::Initialize()
+{
+	if (Instance == nullptr)
+	{
+		Instance = NewObject<UDetailsConfig>();
+		Instance->AddToRoot();
+	}
+}
+
+UDetailsConfig* UDetailsConfig::Get()
+{
+	return Instance;
+}
diff --git a/Engine/Source/Editor/PropertyEditor/Private/DetailsViewConfig.h b/Engine/Source/Editor/PropertyEditor/Private/DetailsViewConfig.h
--- a/Engine/Source/Editor/PropertyEditor/Private/DetailsViewConfig.h
+++ b/Engine/Source/Editor/PropertyEditor/Private/DetailsViewConfig.h
@@ -74,4 +74,11 @@ public:
 
 	UPROPERTY(meta=(EditorConfig))
 	TMap<FName, FDetailsViewConfig> Views;
+
+	static void Initialize();
+	static UDetailsConfig* Get();
+
+private:
+
+	static TObjectPtr<UDetailsConfig> Instance;
 };
\ No newline at end of file
diff --git a/Engine/Source/Editor/PropertyEditor/Private/SDetailsViewBase.cpp b/Engine/Source/Editor/PropertyEditor/Private/SDetailsViewBase.cpp
--- a/Engine/Source/Editor/PropertyEditor/Private/SDetailsViewBase.cpp
+++ b/Engine/Source/Editor/PropertyEditor/Private/SDetailsViewBase.cpp
@@ -39,8 +39,8 @@ SDetailsViewBase::SDetailsViewBase() :
 	, NumVisibleTopLevelObjectNodes(0)
 	, bPendingCleanupTimerSet(false)
 {
-	UDetailsConfig* DetailsConfig = GetMutableDefault<UDetailsConfig>();
-	DetailsConfig->LoadEditorConfig();
+	UDetailsConfig::Initialize();
+	UDetailsConfig::Get()->LoadEditorConfig();
 
 	const FDetailsViewConfig* ViewConfig = GetConstViewConfig();
 	if (ViewConfig != nullptr)
@@ -645,8 +645,7 @@ const FDetailsViewConfig* SDetailsViewBase::GetConstViewConfig() const
 		return nullptr;
 	}
 
-	const UDetailsConfig* DetailsConfig = GetDefault<UDetailsConfig>();
-	return DetailsConfig->Views.Find(DetailsViewArgs.ViewIdentifier);
+	return UDetailsConfig::Get()->Views.Find(DetailsViewArgs.ViewIdentifier);
 }
 
 FDetailsViewConfig* SDetailsViewBase::GetMutableViewConfig()
@@ -656,13 +655,12 @@ FDetailsViewConfig* SDetailsViewBase::GetMutableViewConfig()
 		return nullptr;
 	}
 
-	UDetailsConfig* DetailsConfig = GetMutableDefault<UDetailsConfig>();
-	return &DetailsConfig->Views.FindOrAdd(DetailsViewArgs.ViewIdentifier);
+	return &UDetailsConfig::Get()->Views.FindOrAdd(DetailsViewArgs.ViewIdentifier);
 }
 
 void SDetailsViewBase::SaveViewConfig()
 {
-	GetMutableDefault<UDetailsConfig>()->SaveEditorConfig();
+	UDetailsConfig::Get()->SaveEditorConfig();
 }
 
 void SDetailsViewBase::OnShowOnlyModifiedClicked()
diff --git a/Engine/Source/Editor/SceneOutliner/Private/SSceneOutliner.cpp b/Engine/Source/Editor/SceneOutliner/Private/SSceneOutliner.cpp
--- a/Engine/Source/Editor/SceneOutliner/Private/SSceneOutliner.cpp
+++ b/Engine/Source/Editor/SceneOutliner/Private/SSceneOutliner.cpp
@@ -85,8 +85,8 @@ void SSceneOutliner::Construct(const FArguments& InArgs, const FSceneOutlinerIni
 	SortByColumn = FSceneOutlinerBuiltInColumnTypes::Label();
 	SortMode = EColumnSortMode::Ascending;
 
-	UOutlinerConfig* OutlinerConfig = GetMutableDefault<UOutlinerConfig>();
-	OutlinerConfig->LoadEditorConfig();
+	UOutlinerConfig::Initialize();
+	UOutlinerConfig::Get()->LoadEditorConfig();
 
 	const FSceneOutlinerConfig* SceneOutlinerConfig = GetConstConfig();
 
@@ -2316,8 +2316,7 @@ struct FSceneOutlinerConfig* SSceneOutliner::GetMutableConfig()
 		return nullptr;
 	}
 
-	UOutlinerConfig* OutlinerConfig = GetMutableDefault<UOutlinerConfig>();
-	return &OutlinerConfig->Outliners.FindOrAdd(OutlinerIdentifier);
+	return &UOutlinerConfig::Get()->Outliners.FindOrAdd(OutlinerIdentifier);
 }
 
 
@@ -2328,13 +2327,12 @@ const FSceneOutlinerConfig* SSceneOutliner::GetConstConfig() const
 		return nullptr;
 	}
 
-	const UOutlinerConfig* OutlinerConfig = GetDefault<UOutlinerConfig>();
-	return OutlinerConfig->Outliners.Find(OutlinerIdentifier);
+	return UOutlinerConfig::Get()->Outliners.Find(OutlinerIdentifier);
 }
 
 void SSceneOutliner::SaveConfig()
 {
-	GetMutableDefault<UOutlinerConfig>()->SaveEditorConfig();
+	UOutlinerConfig::Get()->SaveEditorConfig();
 }
 
 #undef LOCTEXT_NAMESPACE
\ No newline at end of file
diff --git a/Engine/Source/Editor/SceneOutliner/Private/SceneOutlinerConfig.cpp b/Engine/Source/Editor/SceneOutliner/Private/SceneOutlinerConfig.cpp
new file mode 100644
--- /dev/null
+++ b/Engine/Source/Editor/SceneOutliner/Private/SceneOutlinerConfig.cpp
@@ -0,0 +1,11 @@
+// Copyright Epic Games, Inc. All Rights Reserved.
+
+#include "SceneOutlinerConfig.h"
+
+TObjectPtr<UOutlinerConfig> UOutlinerConfig::Instance = nullptr;
+
+void UOutlinerConfig::Initialize()
+{
+	Instance = NewObject<UOutlinerConfig>(); 
+	Instance->AddToRoot();
+}
\ No newline at end of file
diff --git a/Engine/Source/Editor/SceneOutliner/Private/SceneOutlinerConfig.h b/Engine/Source/Editor/SceneOutliner/Private/SceneOutlinerConfig.h
--- a/Engine/Source/Editor/SceneOutliner/Private/SceneOutlinerConfig.h
+++ b/Engine/Source/Editor/SceneOutliner/Private/SceneOutlinerConfig.h
@@ -32,6 +32,13 @@ class UOutlinerConfig : public UEditorConfigBase
 	
 public:
 
+	static void Initialize();
+	static UOutlinerConfig* Get() { return Instance; }
+
 	UPROPERTY(meta=(EditorConfig))
 	TMap<FName, FSceneOutlinerConfig> Outliners;
+
+private:
+
+	static TObjectPtr<UOutlinerConfig> Instance;
 };
\ No newline at end of file
