Enum properties in EditorConfigs now serialize correctly. Previously we were casting directly to int64*, which in the case of smaller enum types (eg. uint8), would cause more memory to be read and the values to be incorrect.

EditorConfig will now always keep $type keys.

Fixed some unit test failures that hadn't been caught previously, caused by eg. ANY_PACKAGE change.

diff --git a/Engine/Source/Editor/EditorConfig/Private/EditorConfig.cpp b/Engine/Source/Editor/EditorConfig/Private/EditorConfig.cpp
--- a/Engine/Source/Editor/EditorConfig/Private/EditorConfig.cpp
+++ b/Engine/Source/Editor/EditorConfig/Private/EditorConfig.cpp
@@ -286,19 +286,15 @@ void FEditorConfig::ReadValue(const TSharedPtr<FJsonValue>& JsonValue, const FPr
 	}
 	else if (const FEnumProperty* EnumProperty = CastField<FEnumProperty>(Property))
 	{
-		int64* Value = (int64*) DataPtr;
-
 		UEnum* Enum = EnumProperty->GetEnum();
-		if (Enum != nullptr)
+		FString ValueString;
+		if (JsonValue->TryGetString(ValueString))
 		{
-			FString ValueString;
-			if (JsonValue->TryGetString(ValueString))
+			int64 Index = Enum->GetIndexByNameString(ValueString);
+			if (Index != INDEX_NONE)
 			{
-				int64 Index = Enum->GetIndexByNameString(ValueString);
-				if (Index != INDEX_NONE)
-				{
-					*Value = Enum->GetValueByIndex(Index);
-				}
+				int64 Value = Enum->GetValueByIndex(Index);
+				EnumProperty->GetUnderlyingProperty()->SetIntPropertyValue(DataPtr, Value);
 			}
 		}
 		return;
@@ -653,11 +649,9 @@ TSharedPtr<FJsonValue> FEditorConfig::WriteValue(const FProperty* Property, cons
 	}
 	else if (const FEnumProperty* EnumProperty = CastField<FEnumProperty>(Property))
 	{
-		int64* Value = (int64*) DataPtr;
-
 		UEnum* Enum = EnumProperty->GetEnum();
-		FName ValueName = Enum->GetNameByValue(*Value);
-		ResultValue = MakeShared<FJsonValueString>(ValueName.ToString());
+		int64 Value = EnumProperty->GetUnderlyingProperty()->GetSignedIntPropertyValue(DataPtr);
+		FName ValueName = Enum->GetNameByValue(Value);
 	}
 	else if (const FObjectPropertyBase* ObjectProperty = CastField<FObjectPropertyBase>(Property))
 	{
diff --git a/Engine/Source/Editor/EditorConfig/Private/JsonConfig.cpp b/Engine/Source/Editor/EditorConfig/Private/JsonConfig.cpp
--- a/Engine/Source/Editor/EditorConfig/Private/JsonConfig.cpp
+++ b/Engine/Source/Editor/EditorConfig/Private/JsonConfig.cpp
@@ -193,7 +193,7 @@ namespace UE
 
 	bool FJsonConfig::LoadFromString(FStringView Content)
 	{
-		TSharedRef<FJsonStringReader> JsonReader = FJsonStringReader::Create(Content.GetData());
+		TSharedRef<FJsonStringReader> JsonReader = FJsonStringReader::Create(FString(Content));
 		if (!FJsonSerializer::Deserialize(JsonReader.Get(), OverrideObject))
 		{
 			UE_LOG(LogEditorConfig, Error, TEXT("Failed to deserialize JSON string"));
@@ -925,7 +925,7 @@ namespace UE
 				(*AddArray) = AddedValues;
 			}
 		}
-		else
+		else if (AddedValues.Num() > 0)
 		{
 			(*CurrentObject)->SetArrayField(TEXT("+"), AddedValues);
 		}
@@ -939,7 +939,7 @@ namespace UE
 				(*RemoveArray) = RemovedValues;
 			}
 		}
-		else
+		else if (RemovedValues.Num() > 0)
 		{
 			(*CurrentObject)->SetArrayField(TEXT("-"), RemovedValues);
 		}
@@ -1337,6 +1337,28 @@ namespace UE
 		return true;
 	}
 
+	static bool ShouldAlwaysKeep(const FJsonPath& Path)
+	{
+		static const TCHAR* AlwaysKeep[] =
+		{
+			TEXT("$type")
+		};
+
+		if (Path.Length() > 0)
+		{
+			const FJsonPath::FPart& LastPart = Path[Path.Length() - 1];
+			for (const TCHAR* Key : AlwaysKeep)
+			{
+				if (LastPart.Name == Key)
+				{
+					return true;
+				}
+			}
+		}
+
+		return false;
+	}
+
 	bool FJsonConfig::SetJsonValue(const FJsonPath& Path, const TSharedPtr<FJsonValue>& NewValue)
 	{
 		TSharedPtr<FJsonValue> PreviousValue;
@@ -1353,14 +1375,17 @@ namespace UE
 		// this was a valid change in the merged object 
 		// we need to either add or remove it from this object, depending on if it differs from our parent
 		bool bShouldRemove = false;
-		if (ParentConfig.IsValid())
+		if (!ShouldAlwaysKeep(Path))
 		{
-			if (ParentConfig->TryGetJsonValue(Path, ParentValue))
+			if (ParentConfig.IsValid())
 			{
-				if (FJsonValue::CompareEqual(*ParentValue.Get(), *NewValue.Get()))
+				if (ParentConfig->TryGetJsonValue(Path, ParentValue))
 				{
-					// same as inherited, remove it from this
-					bShouldRemove = true;
+					if (FJsonValue::CompareEqual(*ParentValue.Get(), *NewValue.Get()))
+					{
+						// same as inherited, remove it from this
+						bShouldRemove = true;
+					}
 				}
 			}
 		}
diff --git a/Engine/Source/Editor/EditorConfig/Private/Tests/EditorConfigTests.cpp b/Engine/Source/Editor/EditorConfig/Private/Tests/EditorConfigTests.cpp
--- a/Engine/Source/Editor/EditorConfig/Private/Tests/EditorConfigTests.cpp
+++ b/Engine/Source/Editor/EditorConfig/Private/Tests/EditorConfigTests.cpp
@@ -33,6 +33,83 @@ bool FEditorConfigTestKey::operator!=(const FEditorConfigTestKey& Other) const
 	return !(*this == Other);
 }
 
+IMPLEMENT_SIMPLE_AUTOMATION_TEST(FEditorConfigTests_Save_EnumStruct, "EditorConfig.Save.EnumStruct", EAutomationTestFlags::EditorContext | EAutomationTestFlags::EngineFilter)
+
+bool FEditorConfigTests_Save_EnumStruct::RunTest(const FString& Parameters)
+{
+	FEditorConfigTestEnumStruct Actual;
+	Actual.Before = -1;
+	Actual.Enum = EEditorConfigTestEnum::One;
+	Actual.After = -1;
+
+	FEditorConfig Config;
+	Config.SetRootStruct(Actual, FEditorConfig::EPropertyFilter::All);
+
+	FString ActualString;
+	if (!Config.SaveToString(ActualString))
+	{
+		AddError("Failed to save to string.");
+		return false;
+	}
+
+	const FString Expected =
+R"_JSON({ 
+	"$type": "EditorConfigTestEnumStruct",
+	"Before": -1,
+	"Enum": "One",
+	"After": 0
+})_JSON";
+
+	if (!FEditorConfigTestHelpers::AreJsonStringsEquivalent(ActualString, Expected))
+	{
+		AddError(FString::Printf(TEXT("Contents does not match. Expected:\n%s\n\nActual:\n%s"), *Expected, *ActualString));
+		return false;
+	}
+
+	return true;
+}
+
+IMPLEMENT_SIMPLE_AUTOMATION_TEST(FEditorConfigTests_Load_EnumStruct, "EditorConfig.Load.EnumStruct", EAutomationTestFlags::EditorContext | EAutomationTestFlags::EngineFilter)
+
+bool FEditorConfigTests_Load_EnumStruct::RunTest(const FString& Parameters)
+{
+	// there was a bug where enums were reading past the underlying property and returning the wrong value as a result
+	// this test is meant to catch that
+
+	const FString Contents =
+R"_JSON({ 
+	"$type": "EditorConfigTestEnumStruct",
+	"Before": -1,
+	"Enum": "One",
+	"After": -1
+})_JSON";
+
+	FEditorConfig Config;
+	if (!Config.LoadFromString(Contents))
+	{
+		AddError("Failed to load from string.");
+		return false;
+	}
+
+	FEditorConfigTestEnumStruct Struct;
+	if (!Config.TryGetRootStruct(Struct, FEditorConfig::EPropertyFilter::All))
+	{
+		AddError("Failed to load into struct.");
+		return false;
+	}
+
+	FEditorConfigTestEnumStruct Comparison;
+	Comparison.Before = -1;
+	Comparison.Enum = EEditorConfigTestEnum::One;
+	Comparison.After = -1;
+
+	bool bResult = true;
+	bResult &= TestEqual(TEXT("SimpleStruct"), Struct.Before, Comparison.Before);
+	bResult &= TestEqual(TEXT("SimpleStruct"), Struct.Enum, Comparison.Enum);
+	bResult &= TestEqual(TEXT("SimpleStruct"), Struct.After, Comparison.After);
+	return bResult;
+}
+
 IMPLEMENT_SIMPLE_AUTOMATION_TEST(FEditorConfigTests_Load_SimpleStruct_Empty, "EditorConfig.Load.SimpleStruct.Empty", EAutomationTestFlags::EditorContext | EAutomationTestFlags::EngineFilter)
 
 bool FEditorConfigTests_Load_SimpleStruct_Empty::RunTest(const FString& Parameters)
@@ -598,7 +675,13 @@ R"_JSON({
 	"Int": 42,
 	"String": "foo",
 	"Float": 21.0,
-	"Array": [ "foo", "bar" ]
+	"Array": 
+	{
+		"+": [
+			"foo",
+			"bar"
+		]
+	}
 })_JSON";
 
 	if (!FEditorConfigTestHelpers::AreJsonStringsEquivalent(Result, Expected))
@@ -624,7 +707,7 @@ bool FEditorConfigTests_Save_Object_Full::RunTest(const FString& Parameters)
 	Object->Struct.Bool = false;
 	Object->Struct.String = TEXT("foo");
 	Object->Struct.Int = 42;
-	Object->Struct.Float = 42.f;
+	Object->Struct.Float = 21;
 
 
 	FEditorConfig Config;
@@ -640,7 +723,7 @@ bool FEditorConfigTests_Save_Object_Full::RunTest(const FString& Parameters)
 	const FString Expected =
 R"_JSON({
 	"$type": "EditorConfigTestObject",
-	"Object": "/Game/TestPackages/Misc_Assets/StaticMesh/Cube.Cube",
+	"Object": "/Script/Engine.StaticMesh'/Game/TestPackages/Misc_Assets/StaticMesh/Cube.Cube'",
 	"Struct": 
 	{
 		"$type": "EditorConfigTestSimpleStruct",
diff --git a/Engine/Source/Editor/EditorConfig/Private/Tests/EditorConfigTests.h b/Engine/Source/Editor/EditorConfig/Private/Tests/EditorConfigTests.h
--- a/Engine/Source/Editor/EditorConfig/Private/Tests/EditorConfigTests.h
+++ b/Engine/Source/Editor/EditorConfig/Private/Tests/EditorConfigTests.h
@@ -7,6 +7,29 @@
 
 #include "EditorConfigTests.generated.h"
 
+UENUM()
+enum class EEditorConfigTestEnum : uint8
+{
+	Zero,
+	One,
+	Two
+};
+
+USTRUCT()
+struct FEditorConfigTestEnumStruct
+{
+	GENERATED_BODY()
+
+	UPROPERTY()
+	int32 Before = 0;
+
+	UPROPERTY()
+	EEditorConfigTestEnum Enum = EEditorConfigTestEnum::Zero;
+
+	UPROPERTY()
+	int32 After = 0;
+};
+
 USTRUCT()
 struct FEditorConfigTestSimpleStruct
 {
