EditorConfig:
Fix issue with the ResultValue not being set when Enum Properties were serialized

diff --git a/Engine/Source/Editor/EditorConfig/Private/EditorConfig.cpp b/Engine/Source/Editor/EditorConfig/Private/EditorConfig.cpp
--- a/Engine/Source/Editor/EditorConfig/Private/EditorConfig.cpp
+++ b/Engine/Source/Editor/EditorConfig/Private/EditorConfig.cpp
@@ -652,6 +652,7 @@ TSharedPtr<FJsonValue> FEditorConfig::WriteValue(const FProperty* Property, cons
 		UEnum* Enum = EnumProperty->GetEnum();
 		int64 Value = EnumProperty->GetUnderlyingProperty()->GetSignedIntPropertyValue(DataPtr);
 		FName ValueName = Enum->GetNameByValue(Value);
+		ResultValue = MakeShared<FJsonValueString>(ValueName.ToString());
 	}
 	else if (const FObjectPropertyBase* ObjectProperty = CastField<FObjectPropertyBase>(Property))
 	{
