Update FEditorConfig json parsing to support importing string for structs.
Allows FSoftObjectPath property and many other properties to use simple strings.

diff --git a/Engine/Source/Editor/EditorConfig/Private/EditorConfig.cpp b/Engine/Source/Editor/EditorConfig/Private/EditorConfig.cpp
--- a/Engine/Source/Editor/EditorConfig/Private/EditorConfig.cpp
+++ b/Engine/Source/Editor/EditorConfig/Private/EditorConfig.cpp
@@ -315,6 +315,14 @@ void FEditorConfig::ReadValue(const TSharedPtr<FJsonValue>& JsonValue, const FPr
 		{
 			ReadStruct(*ObjectJsonValue, StructProperty->Struct, DataPtr, Owner, EPropertyFilter::All);
 		}
+		else if (JsonValue->Type == EJson::String)
+		{
+			FString ImportTextString;
+			if (JsonValue->TryGetString(ImportTextString))
+			{
+				Property->ImportText_Direct(*ImportTextString, DataPtr, Owner, 0);
+			}
+		}
 		return;
 	}
 	else if (const FArrayProperty* ArrayProperty = CastField<FArrayProperty>(Property))
