Fix shipping configuration not saving TMaps with string keys to json editorconfig

diff --git a/Engine/Source/Editor/EditorConfig/Private/EditorConfig.cpp b/Engine/Source/Editor/EditorConfig/Private/EditorConfig.cpp
--- a/Engine/Source/Editor/EditorConfig/Private/EditorConfig.cpp
+++ b/Engine/Source/Editor/EditorConfig/Private/EditorConfig.cpp
@@ -562,7 +562,8 @@ TSharedPtr<FJsonValue> FEditorConfig::WriteMap(const FMapProperty* MapProperty,
 				if (MapHelper.IsValidIndex(Idx))
 				{
 					FString KeyString;
-					check(JsonKeysArray[Idx]->TryGetString(KeyString));
+					const bool bGetStringSuccess = JsonKeysArray[Idx]->TryGetString(KeyString);
+					check(bGetStringSuccess);
 
 					ResultObject->SetField(KeyString, JsonValuesArray[Idx]);
 				}
