Add support for UE_SAVED_DIR_OVERRIDE to EditorConfigSubsystem

diff --git a/Engine/Source/Editor/EditorConfig/Private/EditorConfigSubystem.cpp b/Engine/Source/Editor/EditorConfig/Private/EditorConfigSubystem.cpp
--- a/Engine/Source/Editor/EditorConfig/Private/EditorConfigSubystem.cpp
+++ b/Engine/Source/Editor/EditorConfig/Private/EditorConfigSubystem.cpp
@@ -16,7 +16,11 @@ void UEditorConfigSubsystem::Initialize(FSubsystemCollectionBase& Collection)
 {
 	AddSearchDirectory(ESearchDirectoryType::Engine, FPaths::Combine(FPaths::EngineConfigDir(), TEXT("Editor"))); // Engine
 	AddSearchDirectory(ESearchDirectoryType::Project, FPaths::Combine(FPaths::ProjectConfigDir(), TEXT("Editor"))); // ProjectName
+#ifdef UE_SAVED_DIR_OVERRIDE
+	AddSearchDirectory(ESearchDirectoryType::User, FPaths::Combine(FPlatformProcess::UserSettingsDir(), TEXT(PREPROCESSOR_TO_STRING(UE_SAVED_DIR_OVERRIDE)), TEXT("Editor"))); // AppData
+#else
 	AddSearchDirectory(ESearchDirectoryType::User, FPaths::Combine(FPlatformProcess::UserSettingsDir(), *FApp::GetEpicProductIdentifier(), TEXT("Editor"))); // AppData
+#endif
 }
 
 void UEditorConfigSubsystem::Deinitialize()
