Fixed broken EditorConfig unit test due to missing enum name.

#rnx

diff --git a/Engine/Source/Editor/EditorConfig/Private/Tests/EditorConfigTests.cpp b/Engine/Source/Editor/EditorConfig/Private/Tests/EditorConfigTests.cpp
--- a/Engine/Source/Editor/EditorConfig/Private/Tests/EditorConfigTests.cpp
+++ b/Engine/Source/Editor/EditorConfig/Private/Tests/EditorConfigTests.cpp
@@ -56,7 +56,7 @@ bool FEditorConfigTests_Save_EnumStruct::RunTest(const FString& Parameters)
 R"_JSON({ 
 	"$type": "EditorConfigTestEnumStruct",
 	"Before": -1,
-	"Enum": "One",
+	"Enum": "EEditorConfigTestEnum::One",
 	"After": 0
 })_JSON";
 
@@ -80,7 +80,7 @@ bool FEditorConfigTests_Load_EnumStruct::RunTest(const FString& Parameters)
 R"_JSON({ 
 	"$type": "EditorConfigTestEnumStruct",
 	"Before": -1,
-	"Enum": "One",
+	"Enum": "EEditorConfigTestEnum::One",
 	"After": -1
 })_JSON";
 
