[EditorConfig] fix for failing tests

diff --git a/Engine/Source/Editor/EditorConfig/Private/Tests/EditorConfigTests.cpp b/Engine/Source/Editor/EditorConfig/Private/Tests/EditorConfigTests.cpp
--- a/Engine/Source/Editor/EditorConfig/Private/Tests/EditorConfigTests.cpp
+++ b/Engine/Source/Editor/EditorConfig/Private/Tests/EditorConfigTests.cpp
@@ -57,7 +57,7 @@ R"_JSON({
 	"$type": "EditorConfigTestEnumStruct",
 	"Before": -1,
 	"Enum": "EEditorConfigTestEnum::One",
-	"After": 0
+	"After": -1
 })_JSON";
 
 	if (!FEditorConfigTestHelpers::AreJsonStringsEquivalent(ActualString, Expected))
@@ -578,7 +578,7 @@ bool FEditorConfigTests_Load_Object::RunTest(const FString& Parameters)
 	const FString Contents = 
 R"_JSON({
 	"$type": "EditorConfigTestObject",
-	"Object": "/Game/TestPackages/Misc_Assets/StaticMesh/Cube.Cube",
+	"Object": "StaticMesh'/Engine/BasicShapes/Cube.Cube'",
 	"Struct": 
 	{
 		"$type": "EditorConfigTestSimpleStruct",
@@ -700,7 +700,13 @@ bool FEditorConfigTests_Save_Object_Full::RunTest(const FString& Parameters)
 	UEditorConfigTestObject* Object = NewObject<UEditorConfigTestObject>();
 	Object->AddToRoot();
 
-	UObject* Cube = StaticLoadObject(UStaticMesh::StaticClass(), nullptr, TEXT("/Game/TestPackages/Misc_Assets/StaticMesh/Cube.Cube"));
+	UObject* Cube = StaticLoadObject(UStaticMesh::StaticClass(), nullptr, TEXT("StaticMesh'/Engine/BasicShapes/Cube.Cube'"));
+	if (Cube == nullptr)
+	{
+		AddError("Unable to find /Engine/BasicShapes/Cube.Cube");
+		return false;
+	}
+
 	Object->Object = Cube;
 	Object->Number = 42;
 	Object->Struct.Array = TArray<FString>{ TEXT("foo"), TEXT("bar") };
@@ -723,7 +729,7 @@ bool FEditorConfigTests_Save_Object_Full::RunTest(const FString& Parameters)
 	const FString Expected =
 R"_JSON({
 	"$type": "EditorConfigTestObject",
-	"Object": "/Script/Engine.StaticMesh'/Game/TestPackages/Misc_Assets/StaticMesh/Cube.Cube'",
+	"Object": "/Script/Engine.StaticMesh'/Engine/BasicShapes/Cube.Cube'",
 	"Struct": 
 	{
 		"$type": "EditorConfigTestSimpleStruct",
