EditorConfigSubsystem no longer crashes when running under -nothreading.

Async() could complete before the result of the call was assigned to the FPendingSave::WasSuccess variable, which would cause OnSaveCompleted() to remove the entry and then immediately after assign to the newly deleted memory. We now remove all completed saves in Tick().

diff --git a/Engine/Source/Editor/EditorConfig/Private/EditorConfigSubystem.cpp b/Engine/Source/Editor/EditorConfig/Private/EditorConfigSubystem.cpp
--- a/Engine/Source/Editor/EditorConfig/Private/EditorConfigSubystem.cpp
+++ b/Engine/Source/Editor/EditorConfig/Private/EditorConfigSubystem.cpp
@@ -25,8 +25,7 @@ void UEditorConfigSubsystem::Initialize(FSubsystemCollectionBase& Collection)
 
 void UEditorConfigSubsystem::Deinitialize()
 {
-	SaveLock.Lock();
-	ON_SCOPE_EXIT { SaveLock.Unlock(); };
+	FScopeLock Lock(&SaveLock);
 
 	// Synchronously save all Pending Saves on exit
 	for (FPendingSave& Save : PendingSaves)
@@ -41,13 +40,12 @@ void UEditorConfigSubsystem::Deinitialize()
 
 void UEditorConfigSubsystem::Tick(float DeltaTime)
 {
-	SaveLock.Lock();
-	ON_SCOPE_EXIT { SaveLock.Unlock(); };
+	FScopeLock Lock(&SaveLock);
 
 	// Allows PendingSaves to be modified while iterating as
 	// the Async below might execute the task immediately
 	// when running in -nothreading mode.
-	for (int Index = 0; Index < PendingSaves.Num(); ++Index)
+	for (int32 Index = 0; Index < PendingSaves.Num(); ++Index)
 	{
 		FPendingSave& Save = PendingSaves[Index];
 		Save.TimeSinceQueued += DeltaTime;
@@ -69,6 +67,15 @@ void UEditorConfigSubsystem::Tick(float DeltaTime)
 				});
 		}
 	}
+
+	for (int32 Index = 0; Index < PendingSaves.Num(); ++Index)
+	{
+		if (PendingSaves[Index].WasSuccess.IsReady())
+		{
+			PendingSaves.RemoveAt(Index);
+			--Index;
+		}
+	}
 }
 
 TStatId UEditorConfigSubsystem::GetStatId() const
@@ -232,8 +239,7 @@ void UEditorConfigSubsystem::SaveConfig(TSharedRef<FEditorConfig> Config)
 		return;
 	}
 
-	SaveLock.Lock();
-	ON_SCOPE_EXIT{ SaveLock.Unlock(); };
+	FScopeLock Lock(&SaveLock);
 
 	FPendingSave* Existing = PendingSaves.FindByPredicate([Config](const FPendingSave& Element)
 		{
@@ -243,7 +249,7 @@ void UEditorConfigSubsystem::SaveConfig(TSharedRef<FEditorConfig> Config)
 	if (Existing != nullptr)
 	{
 		// reset the timer if we're saving within the grace period and no save is already being executed
-		if (!Existing->WasSuccess.IsValid())
+		if (!Existing->WasSuccess.IsReady())
 		{
 			Existing->TimeSinceQueued = 0;
 		}
@@ -259,8 +265,7 @@ void UEditorConfigSubsystem::SaveConfig(TSharedRef<FEditorConfig> Config)
 
 void UEditorConfigSubsystem::OnSaveCompleted(TSharedPtr<FEditorConfig> Config)
 {
-	SaveLock.Lock();
-	ON_SCOPE_EXIT{ SaveLock.Unlock(); };
+	FScopeLock Lock(&SaveLock);
 
 	const int32 Index = PendingSaves.IndexOfByPredicate([Config](const FPendingSave& Element)
 		{
@@ -271,8 +276,6 @@ void UEditorConfigSubsystem::OnSaveCompleted(TSharedPtr<FEditorConfig> Config)
 	{
 		const FPendingSave& PendingSave = PendingSaves[Index];
 		PendingSave.Config->OnSaved();
-
-		PendingSaves.RemoveAt(Index);
 	}
 }
 
