Fixed wrong include: ScopeExit.h -> ScopeLock.h

#rnx

diff --git a/Engine/Source/Editor/EditorConfig/Private/EditorConfigSubystem.cpp b/Engine/Source/Editor/EditorConfig/Private/EditorConfigSubystem.cpp
--- a/Engine/Source/Editor/EditorConfig/Private/EditorConfigSubystem.cpp
+++ b/Engine/Source/Editor/EditorConfig/Private/EditorConfigSubystem.cpp
@@ -5,7 +5,7 @@
 #include "Async/Async.h"
 #include "Misc/App.h"
 #include "Misc/Paths.h"
-#include "Misc/ScopeExit.h"
+#include "Misc/ScopeLock.h"
 
 UEditorConfigSubsystem::UEditorConfigSubsystem()
 {
