EditorConfigSubsystem no longer creates new threads every frame.

This was caused by a logic error, whereby a new thread was spawned 3 seconds after the save was enqueued and then every frame thereafter regardless of whether or not one had been created. We now use EAsyncExecution::TaskGraph as well.

#rnx

diff --git a/Engine/Source/Editor/EditorConfig/Private/EditorConfigSubystem.cpp b/Engine/Source/Editor/EditorConfig/Private/EditorConfigSubystem.cpp
--- a/Engine/Source/Editor/EditorConfig/Private/EditorConfigSubystem.cpp
+++ b/Engine/Source/Editor/EditorConfig/Private/EditorConfigSubystem.cpp
@@ -51,12 +51,13 @@ void UEditorConfigSubsystem::Tick(float DeltaTime)
 		Save.TimeSinceQueued += DeltaTime;
 
 		const float SaveDelaySeconds = 3.0f;
-		if (Save.TimeSinceQueued > SaveDelaySeconds)
+		if (Save.TimeSinceQueued > SaveDelaySeconds && 
+			!Save.WasSuccess.IsValid())
 		{
 			const FString* FilePath = LoadedConfigs.FindKey(Save.Config);
 			check(FilePath != nullptr);
 			
-			Save.WasSuccess = Async(EAsyncExecution::Thread, 
+			Save.WasSuccess = Async(EAsyncExecution::TaskGraph, 
 				[Config = Save.Config, File = *FilePath]()
 				{
 					return Config->SaveToFile(File);
@@ -249,7 +250,7 @@ void UEditorConfigSubsystem::SaveConfig(TSharedRef<FEditorConfig> Config)
 	if (Existing != nullptr)
 	{
 		// reset the timer if we're saving within the grace period and no save is already being executed
-		if (!Existing->WasSuccess.IsReady())
+		if (!Existing->WasSuccess.IsValid())
 		{
 			Existing->TimeSinceQueued = 0;
 		}
