EditorConfig - Adjust error messages to include line number to assist with fixing .json file errors.
#rnx

diff --git a/Engine/Source/Editor/EditorConfig/Private/JsonConfig.cpp b/Engine/Source/Editor/EditorConfig/Private/JsonConfig.cpp
--- a/Engine/Source/Editor/EditorConfig/Private/JsonConfig.cpp
+++ b/Engine/Source/Editor/EditorConfig/Private/JsonConfig.cpp
@@ -211,7 +211,7 @@ namespace UE
 		TSharedRef<FJsonStringReader> JsonReader = FJsonStringReader::Create(FString(Content));
 		if (!FJsonSerializer::Deserialize(JsonReader.Get(), OverrideObject))
 		{
-			UE_LOG(LogEditorConfig, Error, TEXT("Failed to deserialize JSON string"));
+			UE_LOG(LogEditorConfig, Error, TEXT("Failed to deserialize JSON string: %s"), *JsonReader->GetErrorMessage());
 			return false;
 		}
 
