Added a new search directory type: ProjectOverrides.
It's a in between the default for a project and the one user add to is local configs.

Added a way to register a search directory before editor config subsystem is initialized. This usefull if you need to add a search directory before a editor config could be read.

diff --git a/Engine/Source/Editor/EditorConfig/Private/EditorConfigSubystem.cpp b/Engine/Source/Editor/EditorConfig/Private/EditorConfigSubystem.cpp
--- a/Engine/Source/Editor/EditorConfig/Private/EditorConfigSubystem.cpp
+++ b/Engine/Source/Editor/EditorConfig/Private/EditorConfigSubystem.cpp
@@ -3,10 +3,13 @@
 #include "EditorConfigSubsystem.h"
 
 #include "Async/Async.h"
+#include "Editor.h"
 #include "Misc/App.h"
 #include "Misc/Paths.h"
 #include "Misc/ScopeLock.h"
 
+TArray<TPair<UEditorConfigSubsystem::ESearchDirectoryType, FString>> UEditorConfigSubsystem::EarlyRegistredSearchDirectories;
+
 UEditorConfigSubsystem::UEditorConfigSubsystem()
 {
 	
@@ -14,6 +17,8 @@ UEditorConfigSubsystem::UEditorConfigSubsystem()
 
 void UEditorConfigSubsystem::Initialize(FSubsystemCollectionBase& Collection)
 {
+	SearchDirectories.Reserve(4  + EarlyRegistredSearchDirectories.Num());
+
 	AddSearchDirectory(ESearchDirectoryType::Engine, FPaths::Combine(FPaths::EngineConfigDir(), TEXT("Editor"))); // Engine
 	AddSearchDirectory(ESearchDirectoryType::Project, FPaths::Combine(FPaths::ProjectConfigDir(), TEXT("Editor"))); // ProjectName
 #ifdef UE_SAVED_DIR_OVERRIDE
@@ -21,6 +26,13 @@ void UEditorConfigSubsystem::Initialize(FSubsystemCollectionBase& Collection)
 #else
 	AddSearchDirectory(ESearchDirectoryType::User, FPaths::Combine(FPlatformProcess::UserSettingsDir(), *FApp::GetEpicProductIdentifier(), TEXT("Editor"))); // AppData
 #endif
+
+	for (const TPair<ESearchDirectoryType, FString>& Pair : EarlyRegistredSearchDirectories)
+	{
+		AddSearchDirectory(Pair.Key, Pair.Value);
+	}
+
+	EarlyRegistredSearchDirectories.Empty();
 }
 
 void UEditorConfigSubsystem::Deinitialize()
@@ -295,4 +307,11 @@ void UEditorConfigSubsystem::AddSearchDirectory(ESearchDirectoryType Type, FStri
 		}
 		SearchDirectories.Add(NewEntry);
 	}
+}
+
+void UEditorConfigSubsystem::EarlyAddSearchDirectory(ESearchDirectoryType Type, FStringView SearchDir)
+{
+	EarlyRegistredSearchDirectories.Emplace(Type, SearchDir);
+	
+	check(!GEditor || !GEditor->IsInitialized());
 }
\ No newline at end of file
diff --git a/Engine/Source/Editor/EditorConfig/Public/EditorConfigSubsystem.h b/Engine/Source/Editor/EditorConfig/Public/EditorConfigSubsystem.h
--- a/Engine/Source/Editor/EditorConfig/Public/EditorConfigSubsystem.h
+++ b/Engine/Source/Editor/EditorConfig/Public/EditorConfigSubsystem.h
@@ -72,15 +72,25 @@ public:
 	{
 		Engine,
 		Project,
+		ProjectOverrides,
 		User
 	};
 
 	/**
 	 * Append a new config search directory to the given type.
-	 * Engine directories are searched first, then Project, then User.
+	 * Engine directories are searched first, then Project, then ProjectOverrides, then User.
 	 */
 	void AddSearchDirectory(ESearchDirectoryType Type, FStringView SearchDir);
 
+	/**
+	 * Append a new config search directory to the given type early.
+	 * Engine directories are searched first, then Project, then ProjectOverrides, then User.
+	 * 
+	 * Note: this function can only be called before the UEditorConfigSubsystem is initialized otherwise it will assert. 
+	 * This function is useful if there is a need to register a layer of config before any config read could happen.
+	 */
+	static void EarlyAddSearchDirectory(ESearchDirectoryType Type, FStringView SearchDir);
+
 private:
 	void Initialize(FSubsystemCollectionBase& Collection) override;
 	void Deinitialize() override;
@@ -104,6 +114,8 @@ private:
 	TArray<FPendingSave> PendingSaves;
 	TArray<TPair<ESearchDirectoryType, FString>> SearchDirectories;
 	TMap<FString, TSharedPtr<FEditorConfig>> LoadedConfigs;
+
+	static TArray<TPair<ESearchDirectoryType, FString>> EarlyRegistredSearchDirectories;
 };
 
 template <typename TObject>
