Add option to ignore user search directories when loading EditorConfig

diff --git a/Engine/Source/Editor/EditorConfig/Private/EditorConfigSubystem.cpp b/Engine/Source/Editor/EditorConfig/Private/EditorConfigSubystem.cpp
--- a/Engine/Source/Editor/EditorConfig/Private/EditorConfigSubystem.cpp
+++ b/Engine/Source/Editor/EditorConfig/Private/EditorConfigSubystem.cpp
@@ -162,16 +162,29 @@ bool UEditorConfigSubsystem::ReloadConfig(TSharedRef<FEditorConfig> Config)
 	return true;
 }
 
-TSharedRef<FEditorConfig> UEditorConfigSubsystem::FindOrLoadConfig(FStringView ConfigName)
+TSharedRef<FEditorConfig> UEditorConfigSubsystem::FindOrLoadConfig(FStringView ConfigName, ESearchDirectoryType IncludedTypes)
 {
 	checkf(!ConfigName.IsEmpty(), TEXT("Config name cannot be empty!"));
 
 	FString ConfigNameString(ConfigName);
+
+	// Construct a subset of SearchDirectories to be searched based on IncludedTypes
+	TArrayView<TPair<ESearchDirectoryType, FString>> IncludedSearchDirectories = SearchDirectories;
+	int32 LastSearchDirectoryIndex = SearchDirectories.Num() - 1;
+	for (; LastSearchDirectoryIndex >= 0; --LastSearchDirectoryIndex)
+	{
+		if (SearchDirectories[LastSearchDirectoryIndex].Key <= IncludedTypes)
+		{
+			break;
+		}
+	}
+	checkf(LastSearchDirectoryIndex > 0, TEXT("No valid search directories found for IncludedTypes = %i"), (int32)IncludedTypes);
+	IncludedSearchDirectories.LeftInline(LastSearchDirectoryIndex + 1);
 	
 	// look for the config in the final search directory and return if it's loaded
 	// this assumes that the hierarchy of configs is unchanged
 	// ie. given search directories [Foo, Bar], the existence of Bar/X.json is taken to mean that Foo/X.json has been loaded
-	const FString FinalPath = FPaths::Combine(SearchDirectories.Last().Value, ConfigNameString) + TEXT(".json");
+	const FString FinalPath = FPaths::Combine(IncludedSearchDirectories.Last().Value, ConfigNameString) + TEXT(".json");
 
 	const TSharedPtr<FEditorConfig>* FinalConfig = LoadedConfigs.Find(FinalPath);
 	if (FinalConfig != nullptr)
@@ -182,7 +195,7 @@ TSharedRef<FEditorConfig> UEditorConfigSubsystem::FindOrLoadConfig(FStringView C
 	// find or load all configs in all search directories with the given name
 	TSharedPtr<FEditorConfig> Parent;
 
-	for (const TPair<ESearchDirectoryType, FString>& Dir : SearchDirectories)
+	for (const TPair<ESearchDirectoryType, FString>& Dir : IncludedSearchDirectories)
 	{
 		const FString FullPath = FPaths::Combine(Dir.Value, ConfigNameString) + TEXT(".json");
 
diff --git a/Engine/Source/Editor/EditorConfig/Public/EditorConfigSubsystem.h b/Engine/Source/Editor/EditorConfig/Public/EditorConfigSubsystem.h
--- a/Engine/Source/Editor/EditorConfig/Public/EditorConfigSubsystem.h
+++ b/Engine/Source/Editor/EditorConfig/Public/EditorConfigSubsystem.h
@@ -54,11 +54,21 @@ public:
 	 * @param Filter Whether to save all properties, or only ones marked with the EditorConfig metadata.
 	 */
 	bool SaveConfigObject(const UClass* Class, const UObject* Object, FEditorConfig::EPropertyFilter = FEditorConfig::EPropertyFilter::MetadataOnly);
+
+	enum class ESearchDirectoryType : uint8
+	{
+		Engine,
+		Project,
+		ProjectOverrides,
+		User
+	};
 	
 	/** 
 	 * Find a config with the given name that has already been loaded, load it if it hasn't been, or create one with the given name.
+	 * 
+	 * @param IncludedTypes The "lowest" level of search directory types to include, ex. If IncludedTypes is set to Project, only Engine and Project paths will be loaded
 	 */
-	TSharedRef<FEditorConfig> FindOrLoadConfig(FStringView ConfigName);
+	TSharedRef<FEditorConfig> FindOrLoadConfig(FStringView ConfigName, ESearchDirectoryType IncludedTypes = ESearchDirectoryType::User);
 
 	/*
 	 * Save the given config to the location it was loaded.
@@ -68,14 +78,6 @@ public:
 	/** Force reload the given config and all its (current and potential) parents from disk. */
 	bool ReloadConfig(TSharedRef<FEditorConfig> Config);
 
-	enum class ESearchDirectoryType : uint8
-	{
-		Engine,
-		Project,
-		ProjectOverrides,
-		User
-	};
-
 	/**
 	 * Append a new config search directory to the given type.
 	 * Engine directories are searched first, then Project, then ProjectOverrides, then User.
