[Truncation Warnings] Fix and enable truncation warnings for EditorConfig and EditorStyle and EditorWidgets modules

#rnx

diff --git a/Engine/Source/Editor/EditorConfig/EditorConfig.Build.cs b/Engine/Source/Editor/EditorConfig/EditorConfig.Build.cs
--- a/Engine/Source/Editor/EditorConfig/EditorConfig.Build.cs
+++ b/Engine/Source/Editor/EditorConfig/EditorConfig.Build.cs
@@ -4,18 +4,21 @@ namespace UnrealBuildTool.Rules
 {
 	public class EditorConfig : ModuleRules
 	{
-        public EditorConfig(ReadOnlyTargetRules Target) : base(Target)
+		public EditorConfig(ReadOnlyTargetRules Target) : base(Target)
 		{
+			// Enable truncation warnings in this module
+			CppCompileWarningSettings.UnsafeTypeCastWarningLevel = WarningLevel.Warning;
+
 			PrivateDependencyModuleNames.AddRange(
-                new string[]
-                {
-                    "Core",
+				new string[]
+				{
+					"Core",
 					"CoreUObject",
 					"EditorSubsystem",
 					"Engine",
 					"Json",
 					"UnrealEd",
 				} );
-        }
+		}
 	}
 }
diff --git a/Engine/Source/Editor/EditorConfig/Private/EditorConfig.cpp b/Engine/Source/Editor/EditorConfig/Private/EditorConfig.cpp
--- a/Engine/Source/Editor/EditorConfig/Private/EditorConfig.cpp
+++ b/Engine/Source/Editor/EditorConfig/Private/EditorConfig.cpp
@@ -304,7 +304,7 @@ void FEditorConfig::ReadValue(const TSharedPtr<FJsonValue>& JsonValue, const FPr
 		FString ValueString;
 		if (JsonValue->TryGetString(ValueString))
 		{
-			int64 Index = Enum->GetIndexByNameString(ValueString);
+			int32 Index = Enum->GetIndexByNameString(ValueString);
 			if (Index != INDEX_NONE)
 			{
 				int64 Value = Enum->GetValueByIndex(Index);
@@ -627,42 +627,42 @@ TSharedPtr<FJsonValue> FEditorConfig::WriteValue(const FProperty* Property, cons
 	else if (const FInt8Property* Int8Property = CastField<FInt8Property>(Property))
 	{
 		int8* Value = (int8*) DataPtr;
-		ResultValue = MakeShared<FJsonValueNumber>(*Value);
+		ResultValue = MakeShared<FJsonValueNumber>(static_cast<double>(*Value));
 	}
 	else if (const FInt16Property* Int16Property = CastField<FInt16Property>(Property))
 	{
 		int16* Value = (int16*) DataPtr;
-		ResultValue = MakeShared<FJsonValueNumber>(*Value);
+		ResultValue = MakeShared<FJsonValueNumber>(static_cast<double>(*Value));
 	}
 	else if (const FIntProperty* Int32Property = CastField<FIntProperty>(Property))
 	{
 		int32* Value = (int32*) DataPtr;
-		ResultValue = MakeShared<FJsonValueNumber>(*Value);
+		ResultValue = MakeShared<FJsonValueNumber>(static_cast<double>(*Value));
 	}
 	else if (const FInt64Property* Int64Property = CastField<FInt64Property>(Property))
 	{
 		int64* Value = (int64*) DataPtr;
-		ResultValue = MakeShared<FJsonValueNumber>(*Value);
+		ResultValue = MakeShared<FJsonValueNumber>(static_cast<double>(*Value));
 	}
 	else if (const FByteProperty* ByteProperty = CastField<FByteProperty>(Property))
 	{
 		uint8* Value = (uint8*) DataPtr;
-		ResultValue = MakeShared<FJsonValueNumber>(*Value);
+		ResultValue = MakeShared<FJsonValueNumber>(static_cast<double>(*Value));
 	}
 	else if (const FUInt16Property* Uint16Property = CastField<FUInt16Property>(Property))
 	{
 		uint16* Value = (uint16*) DataPtr;
-		ResultValue = MakeShared<FJsonValueNumber>(*Value);
+		ResultValue = MakeShared<FJsonValueNumber>(static_cast<double>(*Value));
 	}
 	else if (const FUInt32Property* Uint32Property = CastField<FUInt32Property>(Property))
 	{
 		uint32* Value = (uint32*) DataPtr;
-		ResultValue = MakeShared<FJsonValueNumber>(*Value);
+		ResultValue = MakeShared<FJsonValueNumber>(static_cast<double>(*Value));
 	}
 	else if (const FUInt64Property* Uint64Property = CastField<FUInt64Property>(Property))
 	{
 		uint64* Value = (uint64*) DataPtr;
-		ResultValue = MakeShared<FJsonValueNumber>(*Value);
+		ResultValue = MakeShared<FJsonValueNumber>(static_cast<double>(*Value));
 	}
 	else if (const FEnumProperty* EnumProperty = CastField<FEnumProperty>(Property))
 	{
diff --git a/Engine/Source/Editor/EditorStyle/EditorStyle.Build.cs b/Engine/Source/Editor/EditorStyle/EditorStyle.Build.cs
--- a/Engine/Source/Editor/EditorStyle/EditorStyle.Build.cs
+++ b/Engine/Source/Editor/EditorStyle/EditorStyle.Build.cs
@@ -4,8 +4,11 @@ using UnrealBuildTool;
 
 public class EditorStyle : ModuleRules
 {
-    public EditorStyle(ReadOnlyTargetRules Target) : base(Target)
+	public EditorStyle(ReadOnlyTargetRules Target) : base(Target)
 	{
+		// Enable truncation warnings in this module
+		CppCompileWarningSettings.UnsafeTypeCastWarningLevel = WarningLevel.Warning;
+
 		if (!Target.bCompileAgainstEditor)
 		{
 			throw new BuildException("Unable to instantiate EditorStyle module for non-editor targets.");
