Updated headers using UnrealCodeFixup to make sure dllstorage is on methods/staticvars instead of types. Part 5/n

diff --git a/Engine/Source/Editor/DetailCustomizations/Public/Customizations/ColorStructCustomization.h b/Engine/Source/Editor/DetailCustomizations/Public/Customizations/ColorStructCustomization.h
--- a/Engine/Source/Editor/DetailCustomizations/Public/Customizations/ColorStructCustomization.h
+++ b/Engine/Source/Editor/DetailCustomizations/Public/Customizations/ColorStructCustomization.h
@@ -17,6 +17,8 @@
 #include "Templates/SharedPointer.h"
 #include "Widgets/SWidget.h"
 
+#define UE_API DETAILCUSTOMIZATIONS_API
+
 class FDetailWidgetRow;
 class IPropertyHandle;
 class IPropertyTypeCustomization;
@@ -31,13 +33,13 @@ struct FPointerEvent;
 /**
  * Base class for color struct customization (FColor,FLinearColor).
  */
-class DETAILCUSTOMIZATIONS_API FColorStructCustomization
+class FColorStructCustomization
 	: public FMathStructCustomization
 {
 public:
 
-	static TSharedRef<IPropertyTypeCustomization> MakeInstance();
-	~FColorStructCustomization();
+	static UE_API TSharedRef<IPropertyTypeCustomization> MakeInstance();
+	UE_API ~FColorStructCustomization();
 
 protected:
 
@@ -58,7 +60,7 @@ protected:
 	 * @param InColor To be filled with the color value used by this struct, or white if this struct is being used to edit multiple values
 	 * @return The result of trying to get the color value
 	 */
-	virtual FPropertyAccess::Result GetColorAsLinear(FLinearColor& InColor) const;
+	UE_API virtual FPropertyAccess::Result GetColorAsLinear(FLinearColor& InColor) const;
 
 	/**
 	 * Does the type have Alpha Support
@@ -69,30 +71,30 @@ protected:
 	/**
 	 * Stores the colors from the property into SavedPreColorPickerColors
 	*/
-	virtual void GatherSavedPreColorPickerColors();
+	UE_API virtual void GatherSavedPreColorPickerColors();
 
 	/**
 	 * Stores the color as a string in LastPickerColorString
 	*/
-	virtual void SetLastPickerColorString(const FLinearColor NewColor);
+	UE_API virtual void SetLastPickerColorString(const FLinearColor NewColor);
 
 	/**
 	 * Converts Colors into strings
 	 * @param Colors Array of colors to convert
 	 * @return Array of colors converted to string
 	*/
-	virtual TArray<FString> ConvertToPerObjectColors(const TArray<FLinearOrSrgbColor>& Colors) const;
+	UE_API virtual TArray<FString> ConvertToPerObjectColors(const TArray<FLinearOrSrgbColor>& Colors) const;
 
 protected:
 
 	/** Creates the color widget that when clicked spawns the color picker window. */
-	TSharedRef<SWidget> CreateColorWidget(TWeakPtr<IPropertyHandle>);
+	UE_API TSharedRef<SWidget> CreateColorWidget(TWeakPtr<IPropertyHandle>);
 
 	/**
 	 * Does this struct have multiple values?
 	 * @return EVisibility::Visible if it does, EVisibility::Collapsed otherwise
 	 */
-	EVisibility GetMultipleValuesTextVisibility() const;
+	UE_API EVisibility GetMultipleValuesTextVisibility() const;
 
 	/**
 	 * Creates a new color picker for interactively selecting the color
@@ -100,22 +102,22 @@ protected:
 	 * @param bUseAlpha If true alpha will be displayed, otherwise it is ignored
 	 * @param bOnlyRefreshOnOk If true the value of the property will only be refreshed when the user clicks OK in the color picker
 	 */
-	void CreateColorPicker(bool bUseAlpha);
+	UE_API void CreateColorPicker(bool bUseAlpha);
 	
 	/** Creates a new color picker for interactively selecting color */
-	TSharedRef<SColorPicker> CreateInlineColorPicker(TWeakPtr<IPropertyHandle>);
+	UE_API TSharedRef<SColorPicker> CreateInlineColorPicker(TWeakPtr<IPropertyHandle>);
 
 	/**
 	 * Called when the property is set from the color picker 
 	 * 
 	 * @param NewColor The new color value
 	 */
-	void OnSetColorFromColorPicker(FLinearColor NewColor);
+	UE_API void OnSetColorFromColorPicker(FLinearColor NewColor);
 	
 	/**
 	 * Called to reset all colors to before the color picker spawned
 	 */
-	void ResetColors();
+	UE_API void ResetColors();
 
 	/**
 	 * Called when the user clicks cancel in the color picker
@@ -123,52 +125,52 @@ protected:
 	 *
 	 * @param OriginalColor Original color of the property
 	 */
-	void OnColorPickerCancelled(FLinearColor OriginalColor);
+	UE_API void OnColorPickerCancelled(FLinearColor OriginalColor);
 
 	/**
 	 * Called when the color picker window is clsoed
 	 */
-	void OnColorPickerWindowClosed(const TSharedRef<SWindow>& Window);
+	UE_API void OnColorPickerWindowClosed(const TSharedRef<SWindow>& Window);
 
 	/**
 	 * Called when the user enters an interactive color change (dragging something in the picker)
 	 */
-	void OnColorPickerInteractiveBegin();
+	UE_API void OnColorPickerInteractiveBegin();
 
 	/**
 	 * Called when the user completes an interactive color change (dragging something in the picker)
 	 */
-	void OnColorPickerInteractiveEnd();
+	UE_API void OnColorPickerInteractiveEnd();
 
 	/**
 	 * @return The color that should be displayed in the color block                                                              
 	 */
-	FLinearColor OnGetColorForColorBlock() const;
+	UE_API FLinearColor OnGetColorForColorBlock() const;
 	
 	/**
 	 * @return The color that should be displayed in the color block in slate color format                                                           
 	 */
-	FSlateColor  OnGetSlateColorForBlock() const;
+	UE_API FSlateColor  OnGetSlateColorForBlock() const;
 
 	/**
 	 * @return The border color encompassing the entire color block                                                         
 	 */
-	FSlateColor GetColorWidgetBorderColor() const;
+	UE_API FSlateColor GetColorWidgetBorderColor() const;
 	/**
 	 * Called when the user clicks in the color block (opens inline color picker)
 	 */
-	FReply OnMouseButtonDownColorBlock(const FGeometry& MyGeometry, const FPointerEvent& MouseEvent);
+	UE_API FReply OnMouseButtonDownColorBlock(const FGeometry& MyGeometry, const FPointerEvent& MouseEvent);
 	
 	/** Called when the user clicks on the the button to get the full color picker */
-	FReply OnOpenFullColorPickerClicked();
+	UE_API FReply OnOpenFullColorPickerClicked();
 
 protected:
 
 	// FMathStructCustomization interface
 
-	virtual void CustomizeHeader(TSharedRef<IPropertyHandle> StructPropertyHandle, FDetailWidgetRow& HeaderRow, IPropertyTypeCustomizationUtils& StructCustomizationUtils) override;
-	virtual void MakeHeaderRow(TSharedRef<IPropertyHandle>& InStructPropertyHandle, FDetailWidgetRow& Row) override;
-	virtual void GetSortedChildren(TSharedRef<IPropertyHandle> StructPropertyHandle, TArray< TSharedRef<IPropertyHandle> >& OutChildren) override;
+	UE_API virtual void CustomizeHeader(TSharedRef<IPropertyHandle> StructPropertyHandle, FDetailWidgetRow& HeaderRow, IPropertyTypeCustomizationUtils& StructCustomizationUtils) override;
+	UE_API virtual void MakeHeaderRow(TSharedRef<IPropertyHandle>& InStructPropertyHandle, FDetailWidgetRow& Row) override;
+	UE_API virtual void GetSortedChildren(TSharedRef<IPropertyHandle> StructPropertyHandle, TArray< TSharedRef<IPropertyHandle> >& OutChildren) override;
 
 protected:
 	/** Stores a linear or srb color without converting between the two. Only one is valid at a time */
@@ -224,3 +226,5 @@ protected:
 
 	TOptional<int32> TransactionIndex;	
 };
+
+#undef UE_API
diff --git a/Engine/Source/Editor/DetailCustomizations/Public/Customizations/CurveTableCustomization.h b/Engine/Source/Editor/DetailCustomizations/Public/Customizations/CurveTableCustomization.h
--- a/Engine/Source/Editor/DetailCustomizations/Public/Customizations/CurveTableCustomization.h
+++ b/Engine/Source/Editor/DetailCustomizations/Public/Customizations/CurveTableCustomization.h
@@ -40,6 +40,8 @@
 #include "Widgets/Views/STableRow.h"
 #include "Widgets/Views/STableViewBase.h"
 
+#define UE_API DETAILCUSTOMIZATIONS_API
+
 class ITableRow;
 class STableViewBase;
 class SWidget;
@@ -49,7 +51,7 @@ struct FRealCurve;
 /**
  * Customizes a DataTable asset to use a dropdown
  */
-class DETAILCUSTOMIZATIONS_API FCurveTableCustomizationLayout : public IPropertyTypeCustomization
+class FCurveTableCustomizationLayout : public IPropertyTypeCustomization
 {
 public:
 	static TSharedRef<IPropertyTypeCustomization> MakeInstance() 
@@ -173,7 +175,7 @@ protected:
 	}
 
 	/** Returns the ListView for the ComboButton */
-	TSharedRef<SWidget> GetListContent();
+	UE_API TSharedRef<SWidget> GetListContent();
 
 	/** Initial refresh, will not modify data */
 	void OnInitialRefresh()
@@ -282,3 +284,5 @@ protected:
 	/** A cached copy of strings to populate the combo box */
 	TArray<TSharedPtr<FString> > RowNames;
 };
+
+#undef UE_API
diff --git a/Engine/Source/Editor/DetailCustomizations/Public/Customizations/MathStructCustomizations.h b/Engine/Source/Editor/DetailCustomizations/Public/Customizations/MathStructCustomizations.h
--- a/Engine/Source/Editor/DetailCustomizations/Public/Customizations/MathStructCustomizations.h
+++ b/Engine/Source/Editor/DetailCustomizations/Public/Customizations/MathStructCustomizations.h
@@ -17,6 +17,8 @@
 #include "Widgets/Input/NumericTypeInterface.h"
 #include "Widgets/SWidget.h"
 
+#define UE_API DETAILCUSTOMIZATIONS_API
+
 class FDetailWidgetRow;
 class IDetailChildrenBuilder;
 class SWidget;
@@ -25,11 +27,11 @@ struct FSlateBrush;
 /**
  * Base class for math struct customization (e.g, vector, rotator, color)                                                              
  */
-class DETAILCUSTOMIZATIONS_API FMathStructCustomization
+class FMathStructCustomization
 	: public IPropertyTypeCustomization
 {
 public:
-	static TSharedRef<IPropertyTypeCustomization> MakeInstance();
+	static UE_API TSharedRef<IPropertyTypeCustomization> MakeInstance();
 
 	/** Notification when the max/min slider values are changed (only apply if SupportDynamicSliderMaxValue or SupportDynamicSliderMinValue are true) */
 	DECLARE_MULTICAST_DELEGATE_FourParams(FOnNumericEntryBoxDynamicSliderMinMaxValueChanged, float, TWeakPtr<SWidget>, bool, bool);
@@ -42,8 +44,8 @@ public:
 	virtual ~FMathStructCustomization() {}
 
 	/** IPropertyTypeCustomization instance */
-	virtual void CustomizeHeader(TSharedRef<IPropertyHandle> StructPropertyHandle, FDetailWidgetRow& HeaderRow, IPropertyTypeCustomizationUtils& StructCustomizationUtils) override;
-	virtual void CustomizeChildren(TSharedRef<IPropertyHandle> StructPropertyHandle, IDetailChildrenBuilder& StructBuilder, IPropertyTypeCustomizationUtils& StructCustomizationUtils) override;
+	UE_API virtual void CustomizeHeader(TSharedRef<IPropertyHandle> StructPropertyHandle, FDetailWidgetRow& HeaderRow, IPropertyTypeCustomizationUtils& StructCustomizationUtils) override;
+	UE_API virtual void CustomizeChildren(TSharedRef<IPropertyHandle> StructPropertyHandle, IDetailChildrenBuilder& StructBuilder, IPropertyTypeCustomizationUtils& StructCustomizationUtils) override;
 
 	/** Return max/min slider value changed delegate (only apply if SupportDynamicSliderMaxValue or SupportDynamicSliderMinValue are true) */
 	FOnNumericEntryBoxDynamicSliderMinMaxValueChanged& GetOnNumericEntryBoxDynamicSliderMaxValueChangedDelegate() { return OnNumericEntryBoxDynamicSliderMaxValueChanged; }
@@ -61,7 +63,7 @@ public:
 	*
 	* @param WeakHandlePtr	Handle to the property that the new value is for
 	*/
-	bool IsValueEnabled(TWeakPtr<IPropertyHandle> WeakHandlePtr) const;
+	UE_API bool IsValueEnabled(TWeakPtr<IPropertyHandle> WeakHandlePtr) const;
 	
 	// Argument struct for ExtractNumericMetadata, to make it easier to add new metadata to extract.
 	template <typename NumericType>
@@ -101,7 +103,7 @@ protected:
 	 * @param StructPropertyHandle	Handle to the struct property
 	 * @param Row	The header row to add widgets to
 	 */
-	virtual void MakeHeaderRow(TSharedRef<IPropertyHandle>& StructPropertyHandle, FDetailWidgetRow& Row);
+	UE_API virtual void MakeHeaderRow(TSharedRef<IPropertyHandle>& StructPropertyHandle, FDetailWidgetRow& Row);
 
 	/**
 	 * Gets the sorted children for the struct
@@ -109,7 +111,7 @@ protected:
 	 * @param StructPropertyHandle	The handle to the struct property
 	 * @param OutChildren			The child array that should be populated in the order that children should be displayed
 	 */
-	virtual void GetSortedChildren(TSharedRef<IPropertyHandle> StructPropertyHandle, TArray< TSharedRef<IPropertyHandle> >& OutChildren);
+	UE_API virtual void GetSortedChildren(TSharedRef<IPropertyHandle> StructPropertyHandle, TArray< TSharedRef<IPropertyHandle> >& OutChildren);
 
 	/**
 	 * Constructs a widget for the property handle 
@@ -117,7 +119,7 @@ protected:
 	 * @param StructurePropertyHandle	handle of the struct property
 	 * @param PropertyHandle	Child handle of the struct property
 	 */
-	virtual TSharedRef<SWidget> MakeChildWidget(
+	UE_API virtual TSharedRef<SWidget> MakeChildWidget(
 		TSharedRef<IPropertyHandle>& StructurePropertyHandle,
 		TSharedRef<IPropertyHandle>& PropertyHandle);
 
@@ -167,27 +169,27 @@ protected:
 	 */
 	template <typename NumericType>
 	UE_DEPRECATED(5.6, "OnGetValueToolTip is deprecated and no longer called.")
-	FText OnGetValueToolTip(TWeakPtr<IPropertyHandle> WeakHandlePtr) const;
+	UE_API FText OnGetValueToolTip(TWeakPtr<IPropertyHandle> WeakHandlePtr) const;
 
 private:
 
-	TOptional<FTextFormat> OnGetValueToolTipTextFormat(TWeakPtr<IPropertyHandle> WeakHandlePtr) const;
+	UE_API TOptional<FTextFormat> OnGetValueToolTipTextFormat(TWeakPtr<IPropertyHandle> WeakHandlePtr) const;
 
 	/** Gets the brush to use for the lock icon. */
-	const FSlateBrush* GetPreserveScaleRatioImage() const;
+	UE_API const FSlateBrush* GetPreserveScaleRatioImage() const;
 
 	/** Gets the checked value of the preserve scale ratio option */
-	ECheckBoxState IsPreserveScaleRatioChecked() const;
+	UE_API ECheckBoxState IsPreserveScaleRatioChecked() const;
 
 	/** Called when the user toggles preserve ratio. */
-	void OnPreserveScaleRatioToggled(ECheckBoxState NewState, TWeakPtr<IPropertyHandle> PropertyHandle);
+	UE_API void OnPreserveScaleRatioToggled(ECheckBoxState NewState, TWeakPtr<IPropertyHandle> PropertyHandle);
 
-	FReply OnNormalizeClicked(TWeakPtr<IPropertyHandle> PropertyHandle);
+	UE_API FReply OnNormalizeClicked(TWeakPtr<IPropertyHandle> PropertyHandle);
 
 private:
 
 	/** Called when a value starts to be changed by a slider */
-	void OnBeginSliderMovement();
+	UE_API void OnBeginSliderMovement();
 
 	/** Called when a value stops being changed by a slider */
 	template<typename NumericType>
@@ -212,4 +214,6 @@ protected:
 
 	/** True if the ratio is locked when scaling occurs (uniform scaling) */
 	bool bPreserveScaleRatio;
-};
\ No newline at end of file
+};
+
+#undef UE_API
diff --git a/Engine/Source/Editor/DetailCustomizations/Public/Customizations/MobilityCustomization.h b/Engine/Source/Editor/DetailCustomizations/Public/Customizations/MobilityCustomization.h
--- a/Engine/Source/Editor/DetailCustomizations/Public/Customizations/MobilityCustomization.h
+++ b/Engine/Source/Editor/DetailCustomizations/Public/Customizations/MobilityCustomization.h
@@ -10,6 +10,8 @@
 #include "Templates/SharedPointer.h"
 #include "UObject/NameTypes.h"
 
+#define UE_API DETAILCUSTOMIZATIONS_API
+
 class FDetailWidgetRow;
 class IDetailCategoryBuilder;
 class IPropertyHandle;
@@ -17,7 +19,7 @@ class IPropertyHandle;
 enum class ECheckBoxState : uint8;
 
 // Helper class to create a Mobility customization for the specified Property in the specified CategoryBuilder.
-class DETAILCUSTOMIZATIONS_API FMobilityCustomization : public IDetailCustomNodeBuilder, public TSharedFromThis<FMobilityCustomization>
+class FMobilityCustomization : public IDetailCustomNodeBuilder, public TSharedFromThis<FMobilityCustomization>
 {
 public:
 	enum : uint8
@@ -27,20 +29,22 @@ public:
 		MovableMobilityBitMask = (1u << EComponentMobility::Movable),
 	};
 
-	FMobilityCustomization(TSharedPtr<IPropertyHandle> InMobilityHandle, uint8 InRestrictedMobilityBits, bool InForLight);
+	UE_API FMobilityCustomization(TSharedPtr<IPropertyHandle> InMobilityHandle, uint8 InRestrictedMobilityBits, bool InForLight);
 
-	virtual void GenerateHeaderRowContent(FDetailWidgetRow& WidgetRow) override;
-	virtual FName GetName() const override;
+	UE_API virtual void GenerateHeaderRowContent(FDetailWidgetRow& WidgetRow) override;
+	UE_API virtual FName GetName() const override;
 	virtual TSharedPtr<IPropertyHandle> GetPropertyHandle() const override { return MobilityHandle; }
 
 private:
 
-	EComponentMobility::Type GetActiveMobility() const;
-	FSlateColor GetMobilityTextColor(EComponentMobility::Type InMobility) const;
-	void OnMobilityChanged(EComponentMobility::Type InMobility);
-	FText GetMobilityToolTip() const;
+	UE_API EComponentMobility::Type GetActiveMobility() const;
+	UE_API FSlateColor GetMobilityTextColor(EComponentMobility::Type InMobility) const;
+	UE_API void OnMobilityChanged(EComponentMobility::Type InMobility);
+	UE_API FText GetMobilityToolTip() const;
 
 	TSharedPtr<IPropertyHandle> MobilityHandle;
 	bool bForLight;
 	uint8 RestrictedMobilityBits;
 };
+
+#undef UE_API
diff --git a/Engine/Source/Editor/DetailCustomizations/Public/Customizations/SlateBrushCustomization.h b/Engine/Source/Editor/DetailCustomizations/Public/Customizations/SlateBrushCustomization.h
--- a/Engine/Source/Editor/DetailCustomizations/Public/Customizations/SlateBrushCustomization.h
+++ b/Engine/Source/Editor/DetailCustomizations/Public/Customizations/SlateBrushCustomization.h
@@ -9,45 +9,47 @@
 #include "PropertyHandle.h"
 #include "Templates/SharedPointer.h"
 
+#define UE_API DETAILCUSTOMIZATIONS_API
+
 class IPropertyHandle;
 class SErrorText;
 
-class DETAILCUSTOMIZATIONS_API FSlateBrushStructCustomization : public IPropertyTypeCustomization
+class FSlateBrushStructCustomization : public IPropertyTypeCustomization
 {
 public:
-	static TSharedRef<IPropertyTypeCustomization> MakeInstance(bool bIncludePreview);
+	static UE_API TSharedRef<IPropertyTypeCustomization> MakeInstance(bool bIncludePreview);
 
-	FSlateBrushStructCustomization(bool bIncludePreview);
+	UE_API FSlateBrushStructCustomization(bool bIncludePreview);
 
 	/** IPropertyTypeCustomization interface */
-	virtual void CustomizeHeader( TSharedRef<IPropertyHandle> StructPropertyHandle, class FDetailWidgetRow& HeaderRow, IPropertyTypeCustomizationUtils& StructCustomizationUtils ) override;
+	UE_API virtual void CustomizeHeader( TSharedRef<IPropertyHandle> StructPropertyHandle, class FDetailWidgetRow& HeaderRow, IPropertyTypeCustomizationUtils& StructCustomizationUtils ) override;
 
-	virtual void CustomizeChildren( TSharedRef<IPropertyHandle> StructPropertyHandle, class IDetailChildrenBuilder& StructBuilder, IPropertyTypeCustomizationUtils& StructCustomizationUtils ) override;              
+	UE_API virtual void CustomizeChildren( TSharedRef<IPropertyHandle> StructPropertyHandle, class IDetailChildrenBuilder& StructBuilder, IPropertyTypeCustomizationUtils& StructCustomizationUtils ) override;              
 
 private:
 	/**
 	 * Get the Slate Brush outline settings property row visibility
 	 */
-	EVisibility GetOutlineSettingsPropertyVisibility() const;
+	UE_API EVisibility GetOutlineSettingsPropertyVisibility() const;
 
 	/**
 	 * Get the Slate Brush tiling property row visibility
 	 */
-	EVisibility GetTilingPropertyVisibility() const;
+	UE_API EVisibility GetTilingPropertyVisibility() const;
 
 	/**
 	 *  Get the Slate Brush margin property row visibility
 	 */
-	EVisibility GetMarginPropertyVisibility() const;
+	UE_API EVisibility GetMarginPropertyVisibility() const;
 
 	/** Callback for determining image size reset button visibility */
-	bool IsImageSizeResetToDefaultVisible(TSharedPtr<IPropertyHandle> PropertyHandle) const;
+	UE_API bool IsImageSizeResetToDefaultVisible(TSharedPtr<IPropertyHandle> PropertyHandle) const;
 
 	/** Callback for clicking the image size reset button */
-	void OnImageSizeResetToDefault(TSharedPtr<IPropertyHandle> PropertyHandle) const;
+	UE_API void OnImageSizeResetToDefault(TSharedPtr<IPropertyHandle> PropertyHandle) const;
 
 	/** Gets the current default image size based on the current texture resource */
-	FVector2f GetDefaultImageSize() const;
+	UE_API FVector2f GetDefaultImageSize() const;
 
 	/** Slate Brush DrawAs property */
 	TSharedPtr<IPropertyHandle> DrawAsProperty;
@@ -71,3 +73,4 @@ private:
 	bool bIncludePreview;
 };
 
+#undef UE_API
diff --git a/Engine/Source/Editor/DetailCustomizations/Public/Customizations/SlateFontInfoCustomization.h b/Engine/Source/Editor/DetailCustomizations/Public/Customizations/SlateFontInfoCustomization.h
--- a/Engine/Source/Editor/DetailCustomizations/Public/Customizations/SlateFontInfoCustomization.h
+++ b/Engine/Source/Editor/DetailCustomizations/Public/Customizations/SlateFontInfoCustomization.h
@@ -21,73 +21,75 @@
 #include "Widgets/Input/SComboBox.h"
 #include "Widgets/SWidget.h"
 
+#define UE_API DETAILCUSTOMIZATIONS_API
+
 class IPropertyHandle;
 class SWidget;
 struct FAssetData;
 struct FSlateFontInfo;
 
 /** Customize the appearance of an FSlateFontInfo */
-class DETAILCUSTOMIZATIONS_API FSlateFontInfoStructCustomization : public IPropertyTypeCustomization
+class FSlateFontInfoStructCustomization : public IPropertyTypeCustomization
 {
 public:
-	static TSharedRef<IPropertyTypeCustomization> MakeInstance();
+	static UE_API TSharedRef<IPropertyTypeCustomization> MakeInstance();
 
 	/** IPropertyTypeCustomization interface */
-	virtual void CustomizeHeader(TSharedRef<class IPropertyHandle> InStructPropertyHandle, class FDetailWidgetRow& InHeaderRow, IPropertyTypeCustomizationUtils& InStructCustomizationUtils) override;
-	virtual void CustomizeChildren(TSharedRef<class IPropertyHandle> InStructPropertyHandle, class IDetailChildrenBuilder& InStructBuilder, IPropertyTypeCustomizationUtils& InStructCustomizationUtils) override;
+	UE_API virtual void CustomizeHeader(TSharedRef<class IPropertyHandle> InStructPropertyHandle, class FDetailWidgetRow& InHeaderRow, IPropertyTypeCustomizationUtils& InStructCustomizationUtils) override;
+	UE_API virtual void CustomizeChildren(TSharedRef<class IPropertyHandle> InStructPropertyHandle, class IDetailChildrenBuilder& InStructBuilder, IPropertyTypeCustomizationUtils& InStructCustomizationUtils) override;
 
 protected:
-	void AddFontSizeProperty(IDetailChildrenBuilder& InStructBuilder);
-	bool IsFontSizeEnabled() const;
+	UE_API void AddFontSizeProperty(IDetailChildrenBuilder& InStructBuilder);
+	UE_API bool IsFontSizeEnabled() const;
 
 	/** @return The value or unset if properties with multiple values are viewed */
-	TOptional<float> OnFontSizeGetValue() const;
-	void OnFontSizeValueChanged(float NewDisplayValue);
-	void OnFontSizeValueCommitted(float NewDisplayValue, ETextCommit::Type CommitInfo);
+	UE_API TOptional<float> OnFontSizeGetValue() const;
+	UE_API void OnFontSizeValueChanged(float NewDisplayValue);
+	UE_API void OnFontSizeValueCommitted(float NewDisplayValue, ETextCommit::Type CommitInfo);
 
 	/**
 	 * Called when the slider begins to move.  We create a transaction here to undo the property
 	 */
-	void OnFontSizeBeginSliderMovement();
+	UE_API void OnFontSizeBeginSliderMovement();
 
 	/**
 	 * Called when the slider stops moving.  We end the previously created transaction
 	 */
-	void OnFontSizeEndSliderMovement(float NewDisplayValue);
+	UE_API void OnFontSizeEndSliderMovement(float NewDisplayValue);
 
 	/** @return a dynamic text explaining what the size does and showing the current Font DPI setting */
-	FText GetFontSizeTooltipText() const;
+	UE_API FText GetFontSizeTooltipText() const;
 
-	static float ConvertFontSizeFromNativeToDisplay(float FontSize);
-	static float ConvertFontSizeFromDisplayToNative(float FontSize);
+	static UE_API float ConvertFontSizeFromNativeToDisplay(float FontSize);
+	static UE_API float ConvertFontSizeFromDisplayToNative(float FontSize);
 
 	/** Called to filter out invalid font assets */
-	static bool OnFilterFontAsset(const FAssetData& InAssetData);
+	static UE_API bool OnFilterFontAsset(const FAssetData& InAssetData);
 
 	/** Called when the font object used by this FSlateFontInfo has been changed */
-	void OnFontChanged(const FAssetData& InAssetData);
+	UE_API void OnFontChanged(const FAssetData& InAssetData);
 
 	/** Called to see whether the font entry combo should be enabled */
-	bool IsFontEntryComboEnabled() const;
+	UE_API bool IsFontEntryComboEnabled() const;
 
 	/** Called before the font entry combo is opened - used to update the list of available font entries */
-	void OnFontEntryComboOpening();
+	UE_API void OnFontEntryComboOpening();
 
 	/** Called when the selection of the font entry combo is changed */
-	void OnFontEntrySelectionChanged(TSharedPtr<FName> InNewSelection, ESelectInfo::Type);
+	UE_API void OnFontEntrySelectionChanged(TSharedPtr<FName> InNewSelection, ESelectInfo::Type);
 
 	/** Make the widget for an entry in the font entry combo */
-	TSharedRef<SWidget> MakeFontEntryWidget(TSharedPtr<FName> InFontEntry);
+	UE_API TSharedRef<SWidget> MakeFontEntryWidget(TSharedPtr<FName> InFontEntry);
 
 	/** Get the text to use for the font entry combo button */
-	FText GetFontEntryComboText() const;
+	UE_API FText GetFontEntryComboText() const;
 
 	/** Get the name of the currently active font entry (may not be the selected entry if the entry is set to use "None") */
-	FName GetActiveFontEntry() const;
+	UE_API FName GetActiveFontEntry() const;
 
 	/** Get the array of FSlateFontInfo instances this customization is currently editing */
-	TArray<FSlateFontInfo*> GetFontInfoBeingEdited();
-	TArray<const FSlateFontInfo*> GetFontInfoBeingEdited() const;
+	UE_API TArray<FSlateFontInfo*> GetFontInfoBeingEdited();
+	UE_API TArray<const FSlateFontInfo*> GetFontInfoBeingEdited() const;
 
 	/** Handle to the struct property being edited */
 	TSharedPtr<IPropertyHandle> StructPropertyHandle;
@@ -114,3 +116,5 @@ protected:
 	float LastSliderFontSizeCommittedValue = 0.0f;
 
 };
+
+#undef UE_API
diff --git a/Engine/Source/Editor/DetailCustomizations/Public/Customizations/VectorStructCustomization.h b/Engine/Source/Editor/DetailCustomizations/Public/Customizations/VectorStructCustomization.h
--- a/Engine/Source/Editor/DetailCustomizations/Public/Customizations/VectorStructCustomization.h
+++ b/Engine/Source/Editor/DetailCustomizations/Public/Customizations/VectorStructCustomization.h
@@ -6,25 +6,29 @@
 #include "Customizations/MathStructCustomizations.h"
 #include "Templates/SharedPointer.h"
 
+#define UE_API DETAILCUSTOMIZATIONS_API
+
 class IPropertyHandle;
 class IPropertyTypeCustomization;
 
 /**
  * Customizes FVector structs.
  */
-class DETAILCUSTOMIZATIONS_API FVectorStructCustomization
+class FVectorStructCustomization
 	: public FMathStructCustomization
 {
 public:
 
 	/** @return A new instance of this class */
-	static TSharedRef<IPropertyTypeCustomization> MakeInstance();
+	static UE_API TSharedRef<IPropertyTypeCustomization> MakeInstance();
 
 private:
 
 	/** IPropertyTypeCustomization interface */
-	virtual void CustomizeChildren(TSharedRef<IPropertyHandle> StructPropertyHandle, IDetailChildrenBuilder& StructBuilder, IPropertyTypeCustomizationUtils& StructCustomizationUtils) override;
+	UE_API virtual void CustomizeChildren(TSharedRef<IPropertyHandle> StructPropertyHandle, IDetailChildrenBuilder& StructBuilder, IPropertyTypeCustomizationUtils& StructCustomizationUtils) override;
 
 	/** FMathStructCustomization interface */
-	virtual void GetSortedChildren(TSharedRef<IPropertyHandle> StructPropertyHandle, TArray<TSharedRef<IPropertyHandle>>& OutChildren) override;
+	UE_API virtual void GetSortedChildren(TSharedRef<IPropertyHandle> StructPropertyHandle, TArray<TSharedRef<IPropertyHandle>>& OutChildren) override;
 };
+
+#undef UE_API
diff --git a/Engine/Source/Editor/DetailCustomizations/Public/KeyStructCustomization.h b/Engine/Source/Editor/DetailCustomizations/Public/KeyStructCustomization.h
--- a/Engine/Source/Editor/DetailCustomizations/Public/KeyStructCustomization.h
+++ b/Engine/Source/Editor/DetailCustomizations/Public/KeyStructCustomization.h
@@ -10,6 +10,8 @@
 #include "SKeySelector.h"
 #include "Templates/SharedPointer.h"
 
+#define UE_API DETAILCUSTOMIZATIONS_API
+
 class IPropertyHandle;
 class SWidget;
 struct FKey;
@@ -18,23 +20,23 @@ struct FKey;
  * Implements a details panel customization for FKey structures.
  * As  "Key"				<SKeySelector>
  */
-class DETAILCUSTOMIZATIONS_API FKeyStructCustomization
+class FKeyStructCustomization
 	: public IPropertyTypeCustomization
 {
 public:
-	FKeyStructCustomization();
+	UE_API FKeyStructCustomization();
 
 	UE_NONCOPYABLE(FKeyStructCustomization)
 
 	// IPropertyTypeCustomization interface
 
-	virtual void CustomizeHeader(TSharedRef<class IPropertyHandle> StructPropertyHandle, class FDetailWidgetRow& HeaderRow, IPropertyTypeCustomizationUtils& StructCustomizationUtils) override;
+	UE_API virtual void CustomizeHeader(TSharedRef<class IPropertyHandle> StructPropertyHandle, class FDetailWidgetRow& HeaderRow, IPropertyTypeCustomizationUtils& StructCustomizationUtils) override;
 	virtual void CustomizeChildren(TSharedRef<IPropertyHandle> StructPropertyHandle, class IDetailChildrenBuilder& ChildBuilder, IPropertyTypeCustomizationUtils& StructCustomizationUtils) override { };
 	virtual bool ShouldInlineKey() const override { return true; }
 
 	// Helper variant that generates the key struct in the header and appends a single button at the end
 	// TODO: Is there a better way?
-	void CustomizeHeaderOnlyWithButton(TSharedRef<class IPropertyHandle> StructPropertyHandle, class FDetailWidgetRow& HeaderRow, IPropertyTypeCustomizationUtils& StructCustomizationUtils, TSharedRef<SWidget> Button);
+	UE_API void CustomizeHeaderOnlyWithButton(TSharedRef<class IPropertyHandle> StructPropertyHandle, class FDetailWidgetRow& HeaderRow, IPropertyTypeCustomizationUtils& StructCustomizationUtils, TSharedRef<SWidget> Button);
 
 	PRAGMA_DISABLE_DEPRECATION_WARNINGS
 	/** Sets a bool for whether the information Icon for a combo trigger is displayed or not */ 
@@ -50,7 +52,7 @@ public:
 	{
 		return bDisplayIcon;
 	}
-	PRAGMA_ENABLE_DEPRECATION_WARNINGS
+	UE_API PRAGMA_ENABLE_DEPRECATION_WARNINGS
 	
 	/** Sets a bool for whether the key selector should be enabled or not */
 	void SetEnableKeySelector(bool bKeySelectorEnabled);
@@ -86,7 +88,7 @@ public:
 		return DisabledKeySelectorToolTip;
 	}
 
-	void SetKey(const FString& KeyName);
+	UE_API void SetKey(const FString& KeyName);
 	
 public:
 
@@ -95,7 +97,7 @@ public:
 	 *
 	 * @return A new struct customization for Keys.
 	 */
-	static TSharedRef<IPropertyTypeCustomization> MakeInstance( );
+	static UE_API TSharedRef<IPropertyTypeCustomization> MakeInstance( );
 
 protected:
 	/** Whether the information icon for combo triggers is displayed or not */
@@ -112,10 +114,10 @@ protected:
 	FString DefaultKeyName = TEXT("None");
 	
 	/** Gets the current Key being edited. */
-	TOptional<FKey> GetCurrentKey() const;
+	UE_API TOptional<FKey> GetCurrentKey() const;
 
 	/** Updates the property when a new key is selected. */
-	void OnKeyChanged(TSharedPtr<FKey> SelectedKey);
+	UE_API void OnKeyChanged(TSharedPtr<FKey> SelectedKey);
 
 	/** Holds a handle to the property being edited. */
 	TSharedPtr<IPropertyHandle> PropertyHandle;
@@ -123,3 +125,5 @@ protected:
 	/** shared pointer to the Key Selector. */
     TSharedPtr<SKeySelector> KeySelector = nullptr;
 };
+
+#undef UE_API
diff --git a/Engine/Source/Editor/DetailCustomizations/Public/PerPlatformPropertyCustomization.h b/Engine/Source/Editor/DetailCustomizations/Public/PerPlatformPropertyCustomization.h
--- a/Engine/Source/Editor/DetailCustomizations/Public/PerPlatformPropertyCustomization.h
+++ b/Engine/Source/Editor/DetailCustomizations/Public/PerPlatformPropertyCustomization.h
@@ -15,6 +15,8 @@
 #include "UnrealClient.h"
 #include "Widgets/SWidget.h"
 
+#define UE_API DETAILCUSTOMIZATIONS_API
+
 class FDetailWidgetDecl;
 class FDetailWidgetRow;
 class FMenuBuilder;
@@ -52,7 +54,7 @@ struct FPerPlatformPropertyCustomNodeBuilderArgs
 	FName Name;
 };
 
-class DETAILCUSTOMIZATIONS_API FPerPlatformPropertyCustomNodeBuilder : public IDetailCustomNodeBuilder, public TSharedFromThis<FPerPlatformPropertyCustomNodeBuilder>
+class FPerPlatformPropertyCustomNodeBuilder : public IDetailCustomNodeBuilder, public TSharedFromThis<FPerPlatformPropertyCustomNodeBuilder>
 {
 public:
 	FPerPlatformPropertyCustomNodeBuilder(const FPerPlatformPropertyCustomNodeBuilderArgs& InArgs)
@@ -62,16 +64,16 @@ public:
 	virtual void Tick(float DeltaTime) override {}
 	virtual bool RequiresTick() const override { return false; }
 
-	virtual void SetOnRebuildChildren(FSimpleDelegate InOnRegenerateChildren) override;
-	virtual void SetOnToggleExpansion(FOnToggleNodeExpansion InOnRegenerateChildren) override;
-	virtual void GenerateHeaderRowContent(FDetailWidgetRow& HeaderRow);
-	virtual void GenerateChildContent(IDetailChildrenBuilder& ChildrenBuilder) override;
+	UE_API virtual void SetOnRebuildChildren(FSimpleDelegate InOnRegenerateChildren) override;
+	UE_API virtual void SetOnToggleExpansion(FOnToggleNodeExpansion InOnRegenerateChildren) override;
+	UE_API virtual void GenerateHeaderRowContent(FDetailWidgetRow& HeaderRow);
+	UE_API virtual void GenerateChildContent(IDetailChildrenBuilder& ChildrenBuilder) override;
 	virtual bool InitiallyCollapsed() const override { return true; }
-	virtual FName GetName() const override;
+	UE_API virtual FName GetName() const override;
 private:
-	void OnAddPlatformOverride(FName PlatformName);
-	bool OnRemovePlatformOverride(FName PlatformName);
-	void AddPlatformToMenu(const FName PlatformName, const FTextFormat Format, FMenuBuilder& AddPlatformMenuBuilder);
+	UE_API void OnAddPlatformOverride(FName PlatformName);
+	UE_API bool OnRemovePlatformOverride(FName PlatformName);
+	UE_API void AddPlatformToMenu(const FName PlatformName, const FTextFormat Format, FMenuBuilder& AddPlatformMenuBuilder);
 private:
 	/** Handle to the default value */
 	FPerPlatformPropertyCustomNodeBuilderArgs Args;
@@ -110,3 +112,4 @@ private:
 	TWeakPtr<IPropertyUtilities> PropertyUtilities;
 };
 
+#undef UE_API
diff --git a/Engine/Source/Editor/DetailCustomizations/Public/SplineMetadataDetailsFactory.h b/Engine/Source/Editor/DetailCustomizations/Public/SplineMetadataDetailsFactory.h
--- a/Engine/Source/Editor/DetailCustomizations/Public/SplineMetadataDetailsFactory.h
+++ b/Engine/Source/Editor/DetailCustomizations/Public/SplineMetadataDetailsFactory.h
@@ -18,8 +18,8 @@ public:
 	virtual void GenerateChildContent(IDetailGroup& InGroup) = 0;
 };
 
-UCLASS(Abstract)
-class DETAILCUSTOMIZATIONS_API USplineMetadataDetailsFactoryBase : public UObject
+UCLASS(MinimalAPI, Abstract)
+class USplineMetadataDetailsFactoryBase : public UObject
 {
 	GENERATED_UCLASS_BODY()
 
diff --git a/Engine/Source/Editor/DistCurveEditor/Public/IDistCurveEditor.h b/Engine/Source/Editor/DistCurveEditor/Public/IDistCurveEditor.h
--- a/Engine/Source/Editor/DistCurveEditor/Public/IDistCurveEditor.h
+++ b/Engine/Source/Editor/DistCurveEditor/Public/IDistCurveEditor.h
@@ -11,7 +11,7 @@ class UInterpCurveEdSetup;
    FCurveEdNotifyInterface
 -----------------------------------------------------------------------------*/
 
-class DISTCURVEEDITOR_API FCurveEdNotifyInterface
+class FCurveEdNotifyInterface
 {
 public:
 	virtual void PreEditCurve(TArray<UObject*> CurvesAboutToChange) {}
@@ -35,7 +35,7 @@ public:
    IDistributionCurveEditor
 -----------------------------------------------------------------------------*/
 
-class DISTCURVEEDITOR_API IDistributionCurveEditor : public SCompoundWidget
+class IDistributionCurveEditor : public SCompoundWidget
 {
 public:
 	/** Refreshes the viewport */
diff --git a/Engine/Source/Editor/EditorConfig/Public/EditorConfig.h b/Engine/Source/Editor/EditorConfig/Public/EditorConfig.h
--- a/Engine/Source/Editor/EditorConfig/Public/EditorConfig.h
+++ b/Engine/Source/Editor/EditorConfig/Public/EditorConfig.h
@@ -16,6 +16,8 @@
 #include "Templates/UnrealTemplate.h"
 #include "Templates/UnrealTypeTraits.h"
 
+#define UE_API EDITORCONFIG_API
+
 class FArrayProperty;
 class FJsonObject;
 class FJsonValue;
@@ -27,7 +29,7 @@ class UClass;
 class UObject;
 class UStruct;
 
-class EDITORCONFIG_API FEditorConfig
+class FEditorConfig
 {
 public:
 	enum class EPropertyFilter
@@ -36,13 +38,13 @@ public:
 		MetadataOnly
 	};
 
-	FEditorConfig();
+	UE_API FEditorConfig();
 
-	void SetParent(TSharedPtr<FEditorConfig> InConfig);
+	UE_API void SetParent(TSharedPtr<FEditorConfig> InConfig);
 
-	bool LoadFromString(FStringView Content);
-	bool LoadFromFile(FStringView FilePath);
-	bool SaveToString(FString& OutResult) const;
+	UE_API bool LoadFromString(FStringView Content);
+	UE_API bool LoadFromFile(FStringView FilePath);
+	UE_API bool SaveToString(FString& OutResult) const;
 
 	bool IsValid() const { return JsonConfig.IsValid() && JsonConfig->IsValid(); }
 
@@ -59,8 +61,8 @@ public:
 	template <typename T>
 	bool TryGetRootUObject(T* OutValue, EPropertyFilter Filter = EPropertyFilter::MetadataOnly) const;
 
-	bool TryGetRootStruct(const UStruct* Class, void* OutValue, EPropertyFilter Filter = EPropertyFilter::MetadataOnly) const;
-	bool TryGetRootUObject(const UClass* Class, UObject* OutValue, EPropertyFilter Filter = EPropertyFilter::MetadataOnly) const;
+	UE_API bool TryGetRootStruct(const UStruct* Class, void* OutValue, EPropertyFilter Filter = EPropertyFilter::MetadataOnly) const;
+	UE_API bool TryGetRootUObject(const UClass* Class, UObject* OutValue, EPropertyFilter Filter = EPropertyFilter::MetadataOnly) const;
 
 	template <typename T>
 	void SetStruct(FStringView Key, const T& InValue, EPropertyFilter Filter = EPropertyFilter::MetadataOnly);
@@ -72,12 +74,12 @@ public:
 	template <typename T>
 	void SetRootUObject(const T& InValue, EPropertyFilter Filter = EPropertyFilter::MetadataOnly);
 
-	void SetRootStruct(const UStruct* Class, const void* Instance, EPropertyFilter Filter = EPropertyFilter::MetadataOnly);
-	void SetRootUObject(const UClass* Class, const UObject* Instance, EPropertyFilter Filter = EPropertyFilter::MetadataOnly);
+	UE_API void SetRootStruct(const UStruct* Class, const void* Instance, EPropertyFilter Filter = EPropertyFilter::MetadataOnly);
+	UE_API void SetRootUObject(const UClass* Class, const UObject* Instance, EPropertyFilter Filter = EPropertyFilter::MetadataOnly);
 
-	bool HasOverride(FStringView Key) const;
+	UE_API bool HasOverride(FStringView Key) const;
 
-	void OnSaved();
+	UE_API void OnSaved();
 
 	DECLARE_EVENT_OneParam(FEditorConfig, FOnEditorConfigDirtied, const FEditorConfig&);
 	FOnEditorConfigDirtied& OnEditorConfigDirtied() { return EditorConfigDirtiedEvent; }
@@ -86,20 +88,20 @@ private:
 
 	friend class UEditorConfigSubsystem; // for access to LoadFromFile and SaveToFile
 
-	static void ReadUObject(const TSharedPtr<FJsonObject>& JsonObject, const UClass* Class, UObject* Instance, EPropertyFilter Filter);
-	static void ReadStruct(const TSharedPtr<FJsonObject>& JsonObject, const UStruct* Struct, void* Instance, UObject* Owner, EPropertyFilter Filter);
-	static void ReadValue(const TSharedPtr<FJsonValue>& JsonValue, const FProperty* Property, void* DataPtr, UObject* Owner);
+	static UE_API void ReadUObject(const TSharedPtr<FJsonObject>& JsonObject, const UClass* Class, UObject* Instance, EPropertyFilter Filter);
+	static UE_API void ReadStruct(const TSharedPtr<FJsonObject>& JsonObject, const UStruct* Struct, void* Instance, UObject* Owner, EPropertyFilter Filter);
+	static UE_API void ReadValue(const TSharedPtr<FJsonValue>& JsonValue, const FProperty* Property, void* DataPtr, UObject* Owner);
 	
-	static TSharedPtr<FJsonObject> WriteStruct(const UStruct* Struct, const void* Instance, const void* Defaults, EPropertyFilter Filter);
-	static TSharedPtr<FJsonObject> WriteUObject(const UClass* Class, const UObject* Instance, EPropertyFilter Filter);
-	static TSharedPtr<FJsonValue> WriteArray(const FArrayProperty* ArrayProperty, const void* DataPtr);
-	static TSharedPtr<FJsonValue> WriteSet(const FSetProperty* Property, const void* DataPtr);
-	static TSharedPtr<FJsonValue> WriteMap(const FMapProperty* Property, const void* DataPtr);
-	static TSharedPtr<FJsonValue> WriteValue(const FProperty* Property, const void* DataPtr, const void* Defaults);
+	static UE_API TSharedPtr<FJsonObject> WriteStruct(const UStruct* Struct, const void* Instance, const void* Defaults, EPropertyFilter Filter);
+	static UE_API TSharedPtr<FJsonObject> WriteUObject(const UClass* Class, const UObject* Instance, EPropertyFilter Filter);
+	static UE_API TSharedPtr<FJsonValue> WriteArray(const FArrayProperty* ArrayProperty, const void* DataPtr);
+	static UE_API TSharedPtr<FJsonValue> WriteSet(const FSetProperty* Property, const void* DataPtr);
+	static UE_API TSharedPtr<FJsonValue> WriteMap(const FMapProperty* Property, const void* DataPtr);
+	static UE_API TSharedPtr<FJsonValue> WriteValue(const FProperty* Property, const void* DataPtr, const void* Defaults);
 
-	bool SaveToFile(FStringView FilePath) const;
+	UE_API bool SaveToFile(FStringView FilePath) const;
 
-	void SetDirty();
+	UE_API void SetDirty();
 
 private:
 	TSharedPtr<UE::FJsonConfig> JsonConfig;
@@ -238,4 +240,6 @@ void FEditorConfig::SetRootUObject(const T& InValue, EPropertyFilter Filter)
 
 		SetRootUObject(T::StaticClass(), &InValue, Filter);
 	}
-}
\ No newline at end of file
+}
+
+#undef UE_API
diff --git a/Engine/Source/Editor/EditorConfig/Public/EditorConfigBase.h b/Engine/Source/Editor/EditorConfig/Public/EditorConfigBase.h
--- a/Engine/Source/Editor/EditorConfig/Public/EditorConfigBase.h
+++ b/Engine/Source/Editor/EditorConfig/Public/EditorConfigBase.h
@@ -8,17 +8,21 @@
 
 #include "EditorConfigBase.generated.h"
 
+#define UE_API EDITORCONFIG_API
+
 /** Inherit from this class to simplify saving and loading properties from editor configs. */
-UCLASS()
-class EDITORCONFIG_API UEditorConfigBase : public UObject
+UCLASS(MinimalAPI)
+class UEditorConfigBase : public UObject
 {
 	GENERATED_BODY()
 
 public:
 
 	/** Load any properties of this class into properties marked with metadata tag "EditorConfig" from the class's EditorConfig */
-	bool LoadEditorConfig();
+	UE_API bool LoadEditorConfig();
 
 	/** Save any properties of this class in properties marked with metadata tag "EditorConfig" into the class's EditorConfig. */
-	bool SaveEditorConfig() const;
-};
\ No newline at end of file
+	UE_API bool SaveEditorConfig() const;
+};
+
+#undef UE_API
diff --git a/Engine/Source/Editor/EditorConfig/Public/EditorConfigSubsystem.h b/Engine/Source/Editor/EditorConfig/Public/EditorConfigSubsystem.h
--- a/Engine/Source/Editor/EditorConfig/Public/EditorConfigSubsystem.h
+++ b/Engine/Source/Editor/EditorConfig/Public/EditorConfigSubsystem.h
@@ -9,15 +9,17 @@
 
 #include "EditorConfigSubsystem.generated.h"
 
-UCLASS()
-class EDITORCONFIG_API UEditorConfigSubsystem : 
+#define UE_API EDITORCONFIG_API
+
+UCLASS(MinimalAPI)
+class UEditorConfigSubsystem : 
 	public UEditorSubsystem,
 	public FTickableEditorObject
 {
 	GENERATED_BODY()
 
 public:
-	UEditorConfigSubsystem();
+	UE_API UEditorConfigSubsystem();
 
 	/** 
 	 * Load a config from the root of the JSON file into a given UObject.
@@ -44,7 +46,7 @@ public:
 	 * @param Object The UObject instance to load into.
 	 * @param Filter Whether to load all properties, or only ones marked with the EditorConfig metadata.
 	 */
-	bool LoadConfigObject(const UClass* Class, UObject* Object, FEditorConfig::EPropertyFilter = FEditorConfig::EPropertyFilter::MetadataOnly);
+	UE_API bool LoadConfigObject(const UClass* Class, UObject* Object, FEditorConfig::EPropertyFilter = FEditorConfig::EPropertyFilter::MetadataOnly);
 	
 	/** 
 	 * Save the given UObject of the given class to the root of the JSON config.
@@ -53,7 +55,7 @@ public:
 	 * @param Object The UObject instance to save.
 	 * @param Filter Whether to save all properties, or only ones marked with the EditorConfig metadata.
 	 */
-	bool SaveConfigObject(const UClass* Class, const UObject* Object, FEditorConfig::EPropertyFilter = FEditorConfig::EPropertyFilter::MetadataOnly);
+	UE_API bool SaveConfigObject(const UClass* Class, const UObject* Object, FEditorConfig::EPropertyFilter = FEditorConfig::EPropertyFilter::MetadataOnly);
 
 	enum class ESearchDirectoryType : uint8
 	{
@@ -68,21 +70,21 @@ public:
 	 * 
 	 * @param IncludedTypes The "lowest" level of search directory types to include, ex. If IncludedTypes is set to Project, only Engine and Project paths will be loaded
 	 */
-	TSharedRef<FEditorConfig> FindOrLoadConfig(FStringView ConfigName, ESearchDirectoryType IncludedTypes = ESearchDirectoryType::User);
+	UE_API TSharedRef<FEditorConfig> FindOrLoadConfig(FStringView ConfigName, ESearchDirectoryType IncludedTypes = ESearchDirectoryType::User);
 
 	/*
 	 * Save the given config to the location it was loaded.
 	 */
-	void SaveConfig(TSharedRef<FEditorConfig> Config);
+	UE_API void SaveConfig(TSharedRef<FEditorConfig> Config);
 
 	/** Force reload the given config and all its (current and potential) parents from disk. */
-	bool ReloadConfig(TSharedRef<FEditorConfig> Config);
+	UE_API bool ReloadConfig(TSharedRef<FEditorConfig> Config);
 
 	/**
 	 * Append a new config search directory to the given type.
 	 * Engine directories are searched first, then Project, then ProjectOverrides, then User.
 	 */
-	void AddSearchDirectory(ESearchDirectoryType Type, FStringView SearchDir);
+	UE_API void AddSearchDirectory(ESearchDirectoryType Type, FStringView SearchDir);
 
 	/**
 	 * Append a new config search directory to the given type early.
@@ -91,17 +93,17 @@ public:
 	 * Note: this function can only be called before the UEditorConfigSubsystem is initialized otherwise it will assert. 
 	 * This function is useful if there is a need to register a layer of config before any config read could happen.
 	 */
-	static void EarlyAddSearchDirectory(ESearchDirectoryType Type, FStringView SearchDir);
+	static UE_API void EarlyAddSearchDirectory(ESearchDirectoryType Type, FStringView SearchDir);
 
 private:
-	void Initialize(FSubsystemCollectionBase& Collection) override;
-	void Deinitialize() override;
-	void OnSaveCompleted(TSharedPtr<FEditorConfig> Config);
-	void OnEditorConfigDirtied(const FEditorConfig& Config);
+	UE_API void Initialize(FSubsystemCollectionBase& Collection) override;
+	UE_API void Deinitialize() override;
+	UE_API void OnSaveCompleted(TSharedPtr<FEditorConfig> Config);
+	UE_API void OnEditorConfigDirtied(const FEditorConfig& Config);
 
 	/** FTickableEditorObject interface */
-	void Tick(float DeltaTime) override;
-	TStatId GetStatId() const override;
+	UE_API void Tick(float DeltaTime) override;
+	UE_API TStatId GetStatId() const override;
 
 private:
 	struct FPendingSave
@@ -117,7 +119,7 @@ private:
 	TArray<TPair<ESearchDirectoryType, FString>> SearchDirectories;
 	TMap<FString, TSharedPtr<FEditorConfig>> LoadedConfigs;
 
-	static TArray<TPair<ESearchDirectoryType, FString>> EarlyRegistredSearchDirectories;
+	static UE_API TArray<TPair<ESearchDirectoryType, FString>> EarlyRegistredSearchDirectories;
 };
 
 template <typename TObject>
@@ -138,4 +140,6 @@ bool UEditorConfigSubsystem::SaveConfigObject(const TObject* Object, FEditorConf
 
 	const UClass* Class = TObject::StaticClass();
 	return SaveConfigObject(Class, Object);
-}
\ No newline at end of file
+}
+
+#undef UE_API
diff --git a/Engine/Source/Editor/EditorConfig/Public/EditorMetadataOverrides.h b/Engine/Source/Editor/EditorConfig/Public/EditorMetadataOverrides.h
--- a/Engine/Source/Editor/EditorConfig/Public/EditorMetadataOverrides.h
+++ b/Engine/Source/Editor/EditorConfig/Public/EditorMetadataOverrides.h
@@ -18,6 +18,8 @@
 
 #include "EditorMetadataOverrides.generated.h"
 
+#define UE_API EDITORCONFIG_API
+
 class FEditorConfig;
 class FField;
 class FSubsystemCollectionBase;
@@ -77,82 +79,84 @@ enum class EMetadataType
 	String
 };
 
-UCLASS()
-class EDITORCONFIG_API UEditorMetadataOverrides : 
+UCLASS(MinimalAPI)
+class UEditorMetadataOverrides : 
 	public UEditorSubsystem
 { 
 	GENERATED_BODY()
 
 public:
-	UEditorMetadataOverrides();
+	UE_API UEditorMetadataOverrides();
 	virtual ~UEditorMetadataOverrides() {}
 
-	virtual void Initialize(FSubsystemCollectionBase& Collection) override;
+	UE_API virtual void Initialize(FSubsystemCollectionBase& Collection) override;
 
-	bool LoadFromConfig(TSharedPtr<FEditorConfig> Config);
-	void Save();
+	UE_API bool LoadFromConfig(TSharedPtr<FEditorConfig> Config);
+	UE_API void Save();
 
-	EMetadataType GetMetadataType(const FField* Field, FName Key) const;
-	EMetadataType GetMetadataType(const UStruct* Struct, FName Key) const;
+	UE_API EMetadataType GetMetadataType(const FField* Field, FName Key) const;
+	UE_API EMetadataType GetMetadataType(const UStruct* Struct, FName Key) const;
 
-	bool GetStringMetadata(const FField* Field, FName Key, FString& OutValue) const;
-	void SetStringMetadata(const FField* Field, FName Key, FStringView Value);
+	UE_API bool GetStringMetadata(const FField* Field, FName Key, FString& OutValue) const;
+	UE_API void SetStringMetadata(const FField* Field, FName Key, FStringView Value);
 
-	bool GetFloatMetadata(const FField* Field, FName Key, float& OutValue) const;
-	void SetFloatMetadata(const FField* Field, FName Key, float Value);
+	UE_API bool GetFloatMetadata(const FField* Field, FName Key, float& OutValue) const;
+	UE_API void SetFloatMetadata(const FField* Field, FName Key, float Value);
 
-	bool GetIntMetadata(const FField* Field, FName Key, int32& OutValue) const;
-	void SetIntMetadata(const FField* Field, FName Key, int32 Value);
+	UE_API bool GetIntMetadata(const FField* Field, FName Key, int32& OutValue) const;
+	UE_API void SetIntMetadata(const FField* Field, FName Key, int32 Value);
 
-	bool GetBoolMetadata(const FField* Field, FName Key, bool& OutValue) const;
-	void SetBoolMetadata(const FField* Field, FName Key, bool Value);
+	UE_API bool GetBoolMetadata(const FField* Field, FName Key, bool& OutValue) const;
+	UE_API void SetBoolMetadata(const FField* Field, FName Key, bool Value);
 
-	bool GetClassMetadata(const FField* Field, FName Key, UClass*& OutValue) const;
-	void SetClassMetadata(const FField* Field, FName Key, UClass* Value);
+	UE_API bool GetClassMetadata(const FField* Field, FName Key, UClass*& OutValue) const;
+	UE_API void SetClassMetadata(const FField* Field, FName Key, UClass* Value);
 
-	bool GetArrayMetadata(const FField* Field, FName Key, TArray<FString>& OutValue) const;
-	void SetArrayMetadata(const FField* Field, FName Key, const TArray<FString>& Value);
+	UE_API bool GetArrayMetadata(const FField* Field, FName Key, TArray<FString>& OutValue) const;
+	UE_API void SetArrayMetadata(const FField* Field, FName Key, const TArray<FString>& Value);
 
-	void AddToArrayMetadata(const FField* Field, FName Key, const FString& Value);
-	void RemoveFromArrayMetadata(const FField* Field, FName Key, const FString& Value);
+	UE_API void AddToArrayMetadata(const FField* Field, FName Key, const FString& Value);
+	UE_API void RemoveFromArrayMetadata(const FField* Field, FName Key, const FString& Value);
 
-	void RemoveMetadata(const FField* Field, FName Key);
+	UE_API void RemoveMetadata(const FField* Field, FName Key);
 
-	bool GetStringMetadata(const UStruct* Struct, FName Key, FString& OutValue) const;
-	void SetStringMetadata(const UStruct* Struct, FName Key, FStringView Value);
+	UE_API bool GetStringMetadata(const UStruct* Struct, FName Key, FString& OutValue) const;
+	UE_API void SetStringMetadata(const UStruct* Struct, FName Key, FStringView Value);
 
-	bool GetFloatMetadata(const UStruct* Struct, FName Key, float& OutValue) const;
-	void SetFloatMetadata(const UStruct* Struct, FName Key, float Value);
+	UE_API bool GetFloatMetadata(const UStruct* Struct, FName Key, float& OutValue) const;
+	UE_API void SetFloatMetadata(const UStruct* Struct, FName Key, float Value);
 
-	bool GetIntMetadata(const UStruct* Struct, FName Key, int32& OutValue) const;
-	void SetIntMetadata(const UStruct* Struct, FName Key, int32 Value);
+	UE_API bool GetIntMetadata(const UStruct* Struct, FName Key, int32& OutValue) const;
+	UE_API void SetIntMetadata(const UStruct* Struct, FName Key, int32 Value);
 
-	bool GetBoolMetadata(const UStruct* Struct, FName Key, bool& OutValue) const;
-	void SetBoolMetadata(const UStruct* Struct, FName Key, bool Value);
+	UE_API bool GetBoolMetadata(const UStruct* Struct, FName Key, bool& OutValue) const;
+	UE_API void SetBoolMetadata(const UStruct* Struct, FName Key, bool Value);
 
-	bool GetClassMetadata(const UStruct* Struct, FName Key, UClass*& OutValue) const;
-	void SetClassMetadata(const UStruct* Struct, FName Key, UClass* Value);
+	UE_API bool GetClassMetadata(const UStruct* Struct, FName Key, UClass*& OutValue) const;
+	UE_API void SetClassMetadata(const UStruct* Struct, FName Key, UClass* Value);
 
-	bool GetArrayMetadata(const UStruct* Struct, FName Key, TArray<FString>& OutValue) const;
-	void SetArrayMetadata(const UStruct* Struct, FName Key, const TArray<FString>& Value);
+	UE_API bool GetArrayMetadata(const UStruct* Struct, FName Key, TArray<FString>& OutValue) const;
+	UE_API void SetArrayMetadata(const UStruct* Struct, FName Key, const TArray<FString>& Value);
 
-	void AddToArrayMetadata(const UStruct* Struct, FName Key, const FString& Value);
-	void RemoveFromArrayMetadata(const UStruct* Struct, FName Key, const FString& Value);
+	UE_API void AddToArrayMetadata(const UStruct* Struct, FName Key, const FString& Value);
+	UE_API void RemoveFromArrayMetadata(const UStruct* Struct, FName Key, const FString& Value);
 
-	void RemoveMetadata(const UStruct* Struct, FName Key);
+	UE_API void RemoveMetadata(const UStruct* Struct, FName Key);
 
 private:
-	const FMetadataSet* FindFieldMetadata(const FField* Field) const;
-	FMetadataSet* FindFieldMetadata(const FField* Field);
-	FMetadataSet* FindOrAddFieldMetadata(const FField* Field);
+	UE_API const FMetadataSet* FindFieldMetadata(const FField* Field) const;
+	UE_API FMetadataSet* FindFieldMetadata(const FField* Field);
+	UE_API FMetadataSet* FindOrAddFieldMetadata(const FField* Field);
 
-	const FMetadataSet* FindStructMetadata(const UStruct* Struct) const;
-	FMetadataSet* FindStructMetadata(const UStruct* Struct);
-	FMetadataSet* FindOrAddStructMetadata(const UStruct* Struct);
+	UE_API const FMetadataSet* FindStructMetadata(const UStruct* Struct) const;
+	UE_API FMetadataSet* FindStructMetadata(const UStruct* Struct);
+	UE_API FMetadataSet* FindOrAddStructMetadata(const UStruct* Struct);
 
-	void OnCompleted(bool bSuccess);
+	UE_API void OnCompleted(bool bSuccess);
 
 private:
 	TSharedPtr<FEditorConfig> SourceConfig;
 	FMetadataConfig LoadedMetadata;
 };
+
+#undef UE_API
diff --git a/Engine/Source/Editor/EditorConfig/Public/JsonConfig.h b/Engine/Source/Editor/EditorConfig/Public/JsonConfig.h
--- a/Engine/Source/Editor/EditorConfig/Public/JsonConfig.h
+++ b/Engine/Source/Editor/EditorConfig/Public/JsonConfig.h
@@ -23,7 +23,7 @@ class FText;
 
 namespace UE
 {
-	class EDITORCONFIG_API FJsonPath
+	class FJsonPath
 	{
 	public:
 		struct FPart
@@ -32,27 +32,27 @@ namespace UE
 			int32 Index = INDEX_NONE;
 		};
 
-		FJsonPath();
-		FJsonPath(const TCHAR* Path);
-		FJsonPath(FStringView Path);
-		FJsonPath(const FJsonPath& Other);
-		FJsonPath(FJsonPath&& Other);
+		EDITORCONFIG_API FJsonPath();
+		EDITORCONFIG_API FJsonPath(const TCHAR* Path);
+		EDITORCONFIG_API FJsonPath(FStringView Path);
+		EDITORCONFIG_API FJsonPath(const FJsonPath& Other);
+		EDITORCONFIG_API FJsonPath(FJsonPath&& Other);
 
 		bool IsValid() const { return Length() > 0; }
 		const int32 Length() const { return PathParts.Num(); }
 
-		void Append(FStringView Name);
-		void SetArrayIndex(int32 Index);
+		EDITORCONFIG_API void Append(FStringView Name);
+		EDITORCONFIG_API void SetArrayIndex(int32 Index);
 
-		FJsonPath GetSubPath(int32 NumParts) const;
+		EDITORCONFIG_API FJsonPath GetSubPath(int32 NumParts) const;
 
-		FString ToString() const;
+		EDITORCONFIG_API FString ToString() const;
 
 		const FPart& operator[](int32 Idx) const { return PathParts[Idx]; }
 		const TArray<FPart>& GetAll() const { return PathParts; }
 
 	private:
-		void ParsePath(const FString& InPath);
+		EDITORCONFIG_API void ParsePath(const FString& InPath);
 
 	private:
 		TArray<FPart> PathParts;
@@ -60,70 +60,70 @@ namespace UE
 
 	using FJsonValuePair = TPair<TSharedPtr<FJsonValue>, TSharedPtr<FJsonValue>>;
 
-	class EDITORCONFIG_API FJsonConfig
+	class FJsonConfig
 	{
 	public:
-		FJsonConfig();
+		EDITORCONFIG_API FJsonConfig();
 
-		void SetParent(const TSharedPtr<FJsonConfig>& Parent);
+		EDITORCONFIG_API void SetParent(const TSharedPtr<FJsonConfig>& Parent);
 
-		bool LoadFromFile(FStringView FilePath);
-		bool LoadFromString(FStringView Content);
+		EDITORCONFIG_API bool LoadFromFile(FStringView FilePath);
+		EDITORCONFIG_API bool LoadFromString(FStringView Content);
 	
-		bool SaveToFile(FStringView FilePath) const;
-		bool SaveToString(FString& OutResult) const;
+		EDITORCONFIG_API bool SaveToFile(FStringView FilePath) const;
+		EDITORCONFIG_API bool SaveToString(FString& OutResult) const;
 
 		bool IsValid() const { return MergedObject.IsValid(); }
 
-		const FJsonConfig* GetParentConfig() const;
+		EDITORCONFIG_API const FJsonConfig* GetParentConfig() const;
 
 		template <typename T>
 		bool TryGetNumber(const FJsonPath& Path, T& OutValue) const;
-		bool TryGetBool(const FJsonPath& Path, bool& OutValue) const;
-		bool TryGetString(const FJsonPath& Path, FString& OutValue) const;
-		bool TryGetString(const FJsonPath& Path, FName& OutValue) const;
-		bool TryGetString(const FJsonPath& Path, FText& OutValue) const;
-		bool TryGetJsonValue(const FJsonPath& Path, TSharedPtr<FJsonValue>& OutValue) const;
-		bool TryGetJsonObject(const FJsonPath& Path, TSharedPtr<FJsonObject>& OutValue) const;
-		bool TryGetArray(const FJsonPath& Path, TArray<TSharedPtr<FJsonValue>>& OutArray) const;
+		EDITORCONFIG_API bool TryGetBool(const FJsonPath& Path, bool& OutValue) const;
+		EDITORCONFIG_API bool TryGetString(const FJsonPath& Path, FString& OutValue) const;
+		EDITORCONFIG_API bool TryGetString(const FJsonPath& Path, FName& OutValue) const;
+		EDITORCONFIG_API bool TryGetString(const FJsonPath& Path, FText& OutValue) const;
+		EDITORCONFIG_API bool TryGetJsonValue(const FJsonPath& Path, TSharedPtr<FJsonValue>& OutValue) const;
+		EDITORCONFIG_API bool TryGetJsonObject(const FJsonPath& Path, TSharedPtr<FJsonObject>& OutValue) const;
+		EDITORCONFIG_API bool TryGetArray(const FJsonPath& Path, TArray<TSharedPtr<FJsonValue>>& OutArray) const;
 
-		TSharedPtr<FJsonObject> GetRootObject() const;
+		EDITORCONFIG_API TSharedPtr<FJsonObject> GetRootObject() const;
 
 		// these are specializations for arithmetic and string arrays
 		// these could be templated with enable-ifs, 
 		// but it ended up being more lines of incomprehensible template SFINAE than this clear list of types is
-		bool TryGetArray(const FJsonPath& Path, TArray<bool>& OutArray) const;
-		bool TryGetArray(const FJsonPath& Path, TArray<int8>& OutArray) const;
-		bool TryGetArray(const FJsonPath& Path, TArray<int16>& OutArray) const;
-		bool TryGetArray(const FJsonPath& Path, TArray<int32>& OutArray) const;
-		bool TryGetArray(const FJsonPath& Path, TArray<int64>& OutArray) const;
-		bool TryGetArray(const FJsonPath& Path, TArray<uint8>& OutArray) const;
-		bool TryGetArray(const FJsonPath& Path, TArray<uint16>& OutArray) const;
-		bool TryGetArray(const FJsonPath& Path, TArray<uint32>& OutArray) const;
-		bool TryGetArray(const FJsonPath& Path, TArray<uint64>& OutArray) const;
-		bool TryGetArray(const FJsonPath& Path, TArray<float>& OutArray) const;
-		bool TryGetArray(const FJsonPath& Path, TArray<double>& OutArray) const;
-		bool TryGetArray(const FJsonPath& Path, TArray<FString>& OutArray) const;
-		bool TryGetArray(const FJsonPath& Path, TArray<FText>& OutArray) const;
-		bool TryGetArray(const FJsonPath& Path, TArray<FName>& OutArray) const;
-
-		bool TryGetMap(const FJsonPath& Path, TArray<FJsonValuePair>& OutMap) const;
+		EDITORCONFIG_API bool TryGetArray(const FJsonPath& Path, TArray<bool>& OutArray) const;
+		EDITORCONFIG_API bool TryGetArray(const FJsonPath& Path, TArray<int8>& OutArray) const;
+		EDITORCONFIG_API bool TryGetArray(const FJsonPath& Path, TArray<int16>& OutArray) const;
+		EDITORCONFIG_API bool TryGetArray(const FJsonPath& Path, TArray<int32>& OutArray) const;
+		EDITORCONFIG_API bool TryGetArray(const FJsonPath& Path, TArray<int64>& OutArray) const;
+		EDITORCONFIG_API bool TryGetArray(const FJsonPath& Path, TArray<uint8>& OutArray) const;
+		EDITORCONFIG_API bool TryGetArray(const FJsonPath& Path, TArray<uint16>& OutArray) const;
+		EDITORCONFIG_API bool TryGetArray(const FJsonPath& Path, TArray<uint32>& OutArray) const;
+		EDITORCONFIG_API bool TryGetArray(const FJsonPath& Path, TArray<uint64>& OutArray) const;
+		EDITORCONFIG_API bool TryGetArray(const FJsonPath& Path, TArray<float>& OutArray) const;
+		EDITORCONFIG_API bool TryGetArray(const FJsonPath& Path, TArray<double>& OutArray) const;
+		EDITORCONFIG_API bool TryGetArray(const FJsonPath& Path, TArray<FString>& OutArray) const;
+		EDITORCONFIG_API bool TryGetArray(const FJsonPath& Path, TArray<FText>& OutArray) const;
+		EDITORCONFIG_API bool TryGetArray(const FJsonPath& Path, TArray<FName>& OutArray) const;
+
+		EDITORCONFIG_API bool TryGetMap(const FJsonPath& Path, TArray<FJsonValuePair>& OutMap) const;
 
 		// try to set a number - returns false if:
 		// - the path references an object that doesn't exist
 		// - the path references an index in an array that doesn't exist
 		template <typename T>
 		bool SetNumber(const FJsonPath& Path, T Value);
-		bool SetBool(const FJsonPath& Path, bool Value);
-		bool SetString(const FJsonPath& Path, FStringView Value);
-		bool SetString(const FJsonPath& Path, const FText& Value);
-		bool SetJsonValue(const FJsonPath& Path, const TSharedPtr<FJsonValue>& Value);
-		bool SetJsonObject(const FJsonPath& Path, const TSharedPtr<FJsonObject>& Object);
-		bool SetJsonArray(const FJsonPath& Path, const TArray<TSharedPtr<FJsonValue>>& Array);
+		EDITORCONFIG_API bool SetBool(const FJsonPath& Path, bool Value);
+		EDITORCONFIG_API bool SetString(const FJsonPath& Path, FStringView Value);
+		EDITORCONFIG_API bool SetString(const FJsonPath& Path, const FText& Value);
+		EDITORCONFIG_API bool SetJsonValue(const FJsonPath& Path, const TSharedPtr<FJsonValue>& Value);
+		EDITORCONFIG_API bool SetJsonObject(const FJsonPath& Path, const TSharedPtr<FJsonObject>& Object);
+		EDITORCONFIG_API bool SetJsonArray(const FJsonPath& Path, const TArray<TSharedPtr<FJsonValue>>& Array);
 
-		bool SetRootObject(const TSharedPtr<FJsonObject>& Object);
+		EDITORCONFIG_API bool SetRootObject(const TSharedPtr<FJsonObject>& Object);
 
-		bool HasOverride(const FJsonPath& Path) const;
+		EDITORCONFIG_API bool HasOverride(const FJsonPath& Path) const;
 
 	private:
 
@@ -133,15 +133,15 @@ namespace UE
 		template <typename T>
 		bool TryGetNumericArrayHelper(const FJsonPath& Path, TArray<T>& OutArray) const;
 
-		bool SetJsonValueInMerged(const FJsonPath& Path, const TSharedPtr<FJsonValue>& Value);
-		bool SetJsonValueInOverride(const FJsonPath& Path, const TSharedPtr<FJsonValue>& NewValue, const TSharedPtr<FJsonValue>& PreviousValue, const TSharedPtr<FJsonValue>& ParentValue);
-		bool SetArrayValueInOverride(const TSharedPtr<FJsonValue>& CurrentValue, const TArray<TSharedPtr<FJsonValue>>& NewArray, const TSharedPtr<FJsonValue>& ParentValue);
-		bool SetObjectValueInOverride(const TSharedPtr<FJsonObject>& CurrentObject, const TSharedPtr<FJsonObject>& NewObject, const TSharedPtr<FJsonValue>& ParentValue);
+		EDITORCONFIG_API bool SetJsonValueInMerged(const FJsonPath& Path, const TSharedPtr<FJsonValue>& Value);
+		EDITORCONFIG_API bool SetJsonValueInOverride(const FJsonPath& Path, const TSharedPtr<FJsonValue>& NewValue, const TSharedPtr<FJsonValue>& PreviousValue, const TSharedPtr<FJsonValue>& ParentValue);
+		EDITORCONFIG_API bool SetArrayValueInOverride(const TSharedPtr<FJsonValue>& CurrentValue, const TArray<TSharedPtr<FJsonValue>>& NewArray, const TSharedPtr<FJsonValue>& ParentValue);
+		EDITORCONFIG_API bool SetObjectValueInOverride(const TSharedPtr<FJsonObject>& CurrentObject, const TSharedPtr<FJsonObject>& NewObject, const TSharedPtr<FJsonValue>& ParentValue);
 
-		bool RemoveJsonValueFromOverride(const FJsonPath& Path, const TSharedPtr<FJsonValue>& PreviousValue);
+		EDITORCONFIG_API bool RemoveJsonValueFromOverride(const FJsonPath& Path, const TSharedPtr<FJsonValue>& PreviousValue);
 
-		bool MergeThisWithParent();
-		void OnParentConfigChanged();
+		EDITORCONFIG_API bool MergeThisWithParent();
+		EDITORCONFIG_API void OnParentConfigChanged();
 
 	private:
 
diff --git a/Engine/Source/Editor/EditorFramework/Public/Elements/Framework/TypedElementViewportInteraction.h b/Engine/Source/Editor/EditorFramework/Public/Elements/Framework/TypedElementViewportInteraction.h
--- a/Engine/Source/Editor/EditorFramework/Public/Elements/Framework/TypedElementViewportInteraction.h
+++ b/Engine/Source/Editor/EditorFramework/Public/Elements/Framework/TypedElementViewportInteraction.h
@@ -19,6 +19,8 @@
 
 #include "TypedElementViewportInteraction.generated.h"
 
+#define UE_API EDITORFRAMEWORK_API
+
 class UTypedElementSelectionSet;
 struct FInputDeviceState;
 struct FTypedElementHandle;
@@ -34,26 +36,26 @@ enum class ETypedElementViewportInteractionGizmoManipulationType : uint8
 /**
  * Customization used to allow asset editors (such as the level editor) to override the base behavior of viewport interaction.
  */
-class EDITORFRAMEWORK_API FTypedElementViewportInteractionCustomization : public FTypedElementAssetEditorToolkitHostMixin
+class FTypedElementViewportInteractionCustomization : public FTypedElementAssetEditorToolkitHostMixin
 {
 public:
 	virtual ~FTypedElementViewportInteractionCustomization() = default;
 
 	//~ See UTypedElementViewportInteraction for API docs
-	virtual bool GetGizmoPivotLocation(const TTypedElement<ITypedElementWorldInterface>& InElementWorldHandle, const UE::Widget::EWidgetMode InWidgetMode, FVector& OutPivotLocation);
-	virtual void PreGizmoManipulationStarted(TArrayView<const FTypedElementHandle> InElementHandles, const UE::Widget::EWidgetMode InWidgetMode);
-	virtual void GizmoManipulationStarted(const TTypedElement<ITypedElementWorldInterface>& InElementWorldHandle, const UE::Widget::EWidgetMode InWidgetMode);
-	virtual void GizmoManipulationDeltaUpdate(const TTypedElement<ITypedElementWorldInterface>& InElementWorldHandle, const UE::Widget::EWidgetMode InWidgetMode, const EAxisList::Type InDragAxis, const FInputDeviceState& InInputState, const FTransform& InDeltaTransform, const FVector& InPivotLocation);
-	virtual void GizmoManipulationStopped(const TTypedElement<ITypedElementWorldInterface>& InElementWorldHandle, const UE::Widget::EWidgetMode InWidgetMode, const ETypedElementViewportInteractionGizmoManipulationType InManipulationType);
-	virtual void PostGizmoManipulationStopped(TArrayView<const FTypedElementHandle> InElementHandles, const UE::Widget::EWidgetMode InWidgetMode);
-	virtual void MirrorElement(const TTypedElement<ITypedElementWorldInterface>& InElementWorldHandle, const FVector& InMirrorScale, const FVector& InPivotLocation);
-	virtual bool GetFocusBounds(const TTypedElement<ITypedElementWorldInterface>& InElementWorldHandle, FBoxSphereBounds& OutBounds);
+	UE_API virtual bool GetGizmoPivotLocation(const TTypedElement<ITypedElementWorldInterface>& InElementWorldHandle, const UE::Widget::EWidgetMode InWidgetMode, FVector& OutPivotLocation);
+	UE_API virtual void PreGizmoManipulationStarted(TArrayView<const FTypedElementHandle> InElementHandles, const UE::Widget::EWidgetMode InWidgetMode);
+	UE_API virtual void GizmoManipulationStarted(const TTypedElement<ITypedElementWorldInterface>& InElementWorldHandle, const UE::Widget::EWidgetMode InWidgetMode);
+	UE_API virtual void GizmoManipulationDeltaUpdate(const TTypedElement<ITypedElementWorldInterface>& InElementWorldHandle, const UE::Widget::EWidgetMode InWidgetMode, const EAxisList::Type InDragAxis, const FInputDeviceState& InInputState, const FTransform& InDeltaTransform, const FVector& InPivotLocation);
+	UE_API virtual void GizmoManipulationStopped(const TTypedElement<ITypedElementWorldInterface>& InElementWorldHandle, const UE::Widget::EWidgetMode InWidgetMode, const ETypedElementViewportInteractionGizmoManipulationType InManipulationType);
+	UE_API virtual void PostGizmoManipulationStopped(TArrayView<const FTypedElementHandle> InElementHandles, const UE::Widget::EWidgetMode InWidgetMode);
+	UE_API virtual void MirrorElement(const TTypedElement<ITypedElementWorldInterface>& InElementWorldHandle, const FVector& InMirrorScale, const FVector& InPivotLocation);
+	UE_API virtual bool GetFocusBounds(const TTypedElement<ITypedElementWorldInterface>& InElementWorldHandle, FBoxSphereBounds& OutBounds);
 };
 
 /**
  * Utility to hold a typed element handle and its associated world interface and  viewport interaction customization.
  */
-struct EDITORFRAMEWORK_API FTypedElementViewportInteractionElement
+struct FTypedElementViewportInteractionElement
 {
 public:
 	FTypedElementViewportInteractionElement() = default;
@@ -98,8 +100,8 @@ private:
  * A utility to handle higher-level viewport interactions, but default via UTypedElementWorldInterface,
  * but asset editors can customize this behavior via FTypedElementViewportInteractionCustomization.
  */
-UCLASS(Transient)
-class EDITORFRAMEWORK_API UTypedElementViewportInteraction : public UObject, public TTypedElementInterfaceCustomizationRegistry<FTypedElementViewportInteractionCustomization>
+UCLASS(MinimalAPI, Transient)
+class UTypedElementViewportInteraction : public UObject, public TTypedElementInterfaceCustomizationRegistry<FTypedElementViewportInteractionCustomization>
 {
 	GENERATED_BODY()
 
@@ -108,39 +110,41 @@ public:
 	 * Notify that the gizmo is potentially about to start manipulating the transform of the given set of elements.
 	 * @note This list should have been pre-normalized via UTypedElementSelectionSet::GetNormalizedSelection or UTypedElementSelectionSet::GetNormalizedElementList, or FLevelEditorViewportClient::GetElementsToManipulate.
 	 */
-	void BeginGizmoManipulation(FTypedElementListConstRef InElementsToMove, const UE::Widget::EWidgetMode InWidgetMode);
+	UE_API void BeginGizmoManipulation(FTypedElementListConstRef InElementsToMove, const UE::Widget::EWidgetMode InWidgetMode);
 
 	/**
 	 * Notify that the gizmo has manipulated the transform of the given set of elements.
 	 * @note This list should have been pre-normalized via UTypedElementSelectionSet::GetNormalizedSelection or UTypedElementSelectionSet::GetNormalizedElementList, or FLevelEditorViewportClient::GetElementsToManipulate.
 	 */
-	void UpdateGizmoManipulation(FTypedElementListConstRef InElementsToMove, const UE::Widget::EWidgetMode InWidgetMode, const EAxisList::Type InDragAxis, const FInputDeviceState& InInputState, const FTransform& InDeltaTransform);
+	UE_API void UpdateGizmoManipulation(FTypedElementListConstRef InElementsToMove, const UE::Widget::EWidgetMode InWidgetMode, const EAxisList::Type InDragAxis, const FInputDeviceState& InInputState, const FTransform& InDeltaTransform);
 	
 	/**
 	 * Notify that the gizmo has finished manipulating the transform of the given set of elements.
 	 * @note This list should have been pre-normalized via UTypedElementSelectionSet::GetNormalizedSelection or UTypedElementSelectionSet::GetNormalizedElementList, or FLevelEditorViewportClient::GetElementsToManipulate.
 	 */
-	void EndGizmoManipulation(FTypedElementListConstRef InElementsToMove, const UE::Widget::EWidgetMode InWidgetMode, const ETypedElementViewportInteractionGizmoManipulationType InManipulationType);
+	UE_API void EndGizmoManipulation(FTypedElementListConstRef InElementsToMove, const UE::Widget::EWidgetMode InWidgetMode, const ETypedElementViewportInteractionGizmoManipulationType InManipulationType);
 
 	/**
 	 * Apply the given delta to the specified element.
 	 * This performs a similar operation to an in-flight gizmo manipulation, but without any pre/post-change notification.
 	 */
-	void ApplyDeltaToElement(const FTypedElementHandle& InElementHandle, const UE::Widget::EWidgetMode InWidgetMode, const EAxisList::Type InDragAxis, const FInputDeviceState& InInputState, const FTransform& InDeltaTransform);
+	UE_API void ApplyDeltaToElement(const FTypedElementHandle& InElementHandle, const UE::Widget::EWidgetMode InWidgetMode, const EAxisList::Type InDragAxis, const FInputDeviceState& InInputState, const FTransform& InDeltaTransform);
 
 	/**
 	 * Apply the given mirror scale to the specified element.
 	 */
-	void MirrorElement(const FTypedElementHandle& InElementHandle, const FVector& InMirrorScale);
+	UE_API void MirrorElement(const FTypedElementHandle& InElementHandle, const FVector& InMirrorScale);
 
 	/**
 	 *	Calculate the bounds of all elements in the list, to be focused on in the viewport
 	 *	@note This list should have been pre-normalized via UTypedElementSelectionSet::GetNormalizedSelection or UTypedElementSelectionSet::GetNormalizedElementList, or FLevelEditorViewportClient::GetElementsToManipulate.
 	 */
-	bool GetFocusBounds(FTypedElementListConstRef InElements, FBoxSphereBounds& OutBox);
+	UE_API bool GetFocusBounds(FTypedElementListConstRef InElements, FBoxSphereBounds& OutBox);
 private:
 	/**
 	 * Attempt to resolve the selection interface and viewport interaction customization for the given element, if any.
 	 */
-	FTypedElementViewportInteractionElement ResolveViewportInteractionElement(const FTypedElementHandle& InElementHandle) const;
+	UE_API FTypedElementViewportInteractionElement ResolveViewportInteractionElement(const FTypedElementHandle& InElementHandle) const;
 };
+
+#undef UE_API
diff --git a/Engine/Source/Editor/EditorFramework/Public/Elements/Interfaces/TypedElementDetailsInterface.h b/Engine/Source/Editor/EditorFramework/Public/Elements/Interfaces/TypedElementDetailsInterface.h
--- a/Engine/Source/Editor/EditorFramework/Public/Elements/Interfaces/TypedElementDetailsInterface.h
+++ b/Engine/Source/Editor/EditorFramework/Public/Elements/Interfaces/TypedElementDetailsInterface.h
@@ -34,7 +34,7 @@ class UTypedElementDetailsInterface : public UInterface
 	GENERATED_BODY()
 };
 
-class EDITORFRAMEWORK_API ITypedElementDetailsInterface
+class ITypedElementDetailsInterface
 {
 	GENERATED_BODY()
 
diff --git a/Engine/Source/Editor/EditorFramework/Public/Factories/AssetFactoryInterface.h b/Engine/Source/Editor/EditorFramework/Public/Factories/AssetFactoryInterface.h
--- a/Engine/Source/Editor/EditorFramework/Public/Factories/AssetFactoryInterface.h
+++ b/Engine/Source/Editor/EditorFramework/Public/Factories/AssetFactoryInterface.h
@@ -26,7 +26,7 @@ class UAssetFactoryInterface : public UInterface
 	GENERATED_BODY()
 };
 
-class EDITORFRAMEWORK_API IAssetFactoryInterface
+class IAssetFactoryInterface
 {
 	GENERATED_BODY()
 
diff --git a/Engine/Source/Editor/EditorFramework/Public/Subsystems/EditorElementSubsystem.h b/Engine/Source/Editor/EditorFramework/Public/Subsystems/EditorElementSubsystem.h
--- a/Engine/Source/Editor/EditorFramework/Public/Subsystems/EditorElementSubsystem.h
+++ b/Engine/Source/Editor/EditorFramework/Public/Subsystems/EditorElementSubsystem.h
@@ -14,11 +14,13 @@
 
 #include "EditorElementSubsystem.generated.h"
 
+#define UE_API EDITORFRAMEWORK_API
+
 class UObject;
 struct FTypedElementHandle;
 
-UCLASS(Transient)
-class EDITORFRAMEWORK_API UEditorElementSubsystem : public UEditorSubsystem
+UCLASS(MinimalAPI, Transient)
+class UEditorElementSubsystem : public UEditorSubsystem
 {
 	GENERATED_BODY()
 
@@ -27,26 +29,28 @@ public:
 	 * Sets the world transform of the given element handle, if possible.
 	 * @returns false if the world transform could not be set.
 	 */
-	static bool SetElementTransform(FTypedElementHandle InElementHandle, const FTransform& InWorldTransform);
+	static UE_API bool SetElementTransform(FTypedElementHandle InElementHandle, const FTransform& InWorldTransform);
 
 	/**
 	 * Return a normalized selection set like the editor would use for the gizmo manipulations
 	 */
-	static FTypedElementListRef GetEditorNormalizedSelectionSet(const UTypedElementSelectionSet& SelectionSet);
+	static UE_API FTypedElementListRef GetEditorNormalizedSelectionSet(const UTypedElementSelectionSet& SelectionSet);
 
 	/**
 	 * Return only the manipulatable elements from a selection list.
 	 * Require a editor normalized selection list to behave like the editor selection
 	 * Note: A manipulable element is a element that would be move when manipulating the gizmo in the editor
 	 */
-	static FTypedElementListRef GetEditorManipulableElements(const FTypedElementListRef& NormalizedSelection, UE::Widget::EWidgetMode InManipulationType = UE::Widget::WM_Translate, UWorld* InRequiredWorld = nullptr);
+	static UE_API FTypedElementListRef GetEditorManipulableElements(const FTypedElementListRef& NormalizedSelection, UE::Widget::EWidgetMode InManipulationType = UE::Widget::WM_Translate, UWorld* InRequiredWorld = nullptr);
 
 	/**
 	 * Return the most recently selected element that is manipulable.
 	 * Require a editor normalized selection list to behave like the editor selection
 	 * Note: A manipulable element is a element that would be move when manipulating the gizmo in the editor.
 	 */
-	static TTypedElement<ITypedElementWorldInterface> GetLastSelectedEditorManipulableElement(const FTypedElementListRef& NormalizedSelection, UE::Widget::EWidgetMode InManipulationType = UE::Widget::WM_Translate, UWorld* InRequiredWorld = nullptr);
+	static UE_API TTypedElement<ITypedElementWorldInterface> GetLastSelectedEditorManipulableElement(const FTypedElementListRef& NormalizedSelection, UE::Widget::EWidgetMode InManipulationType = UE::Widget::WM_Translate, UWorld* InRequiredWorld = nullptr);
+
+	static UE_API bool IsElementEditorManipulable(const TTypedElement<ITypedElementWorldInterface>& WorldElement, UE::Widget::EWidgetMode InManipulationType = UE::Widget::WM_Translate, UWorld* InRequiredWorld = nullptr);
+};
 
-	static bool IsElementEditorManipulable(const TTypedElement<ITypedElementWorldInterface>& WorldElement, UE::Widget::EWidgetMode InManipulationType = UE::Widget::WM_Translate, UWorld* InRequiredWorld = nullptr);
-};
\ No newline at end of file
+#undef UE_API
diff --git a/Engine/Source/Editor/EditorFramework/Public/Subsystems/PlacementSubsystem.h b/Engine/Source/Editor/EditorFramework/Public/Subsystems/PlacementSubsystem.h
--- a/Engine/Source/Editor/EditorFramework/Public/Subsystems/PlacementSubsystem.h
+++ b/Engine/Source/Editor/EditorFramework/Public/Subsystems/PlacementSubsystem.h
@@ -21,6 +21,8 @@
 
 #include "PlacementSubsystem.generated.h"
 
+#define UE_API EDITORFRAMEWORK_API
+
 class FSubsystemCollectionBase;
 class IAssetFactoryInterface;
 class UAssetFactoryInterface;
@@ -31,7 +33,7 @@ class UObject;
 struct FTypedElementHandle;
 
 USTRUCT()
-struct EDITORFRAMEWORK_API FAssetPlacementInfo
+struct FAssetPlacementInfo
 {
 	GENERATED_BODY()
 
@@ -68,7 +70,7 @@ struct EDITORFRAMEWORK_API FAssetPlacementInfo
 };
 
 USTRUCT()
-struct EDITORFRAMEWORK_API FPlacementOptions
+struct FPlacementOptions
 {
 	GENERATED_BODY()
 	
@@ -88,55 +90,55 @@ struct EDITORFRAMEWORK_API FPlacementOptions
 	bool bIsCreatingPreviewElements = false;
 };
 
-UCLASS(Transient)
-class EDITORFRAMEWORK_API UPlacementSubsystem : public UEditorSubsystem
+UCLASS(MinimalAPI, Transient)
+class UPlacementSubsystem : public UEditorSubsystem
 {
 	GENERATED_BODY()
 
 public:
 	// USubsystem Interface
-	void Initialize(FSubsystemCollectionBase& Collection) override;
-	void Deinitialize() override;
+	UE_API void Initialize(FSubsystemCollectionBase& Collection) override;
+	UE_API void Deinitialize() override;
 	// End USubsystem Interface
 
 	/**
 	 * Places a single asset based on the given FAssetPlacementInfo and FPlacementOptions.
 	 * @returns an array of FTypedElementHandles corresponding to any successfully placed elements.
 	 */
-	TArray<FTypedElementHandle> PlaceAsset(const FAssetPlacementInfo& InPlacementInfo, const FPlacementOptions& InPlacementOptions);
+	UE_API TArray<FTypedElementHandle> PlaceAsset(const FAssetPlacementInfo& InPlacementInfo, const FPlacementOptions& InPlacementOptions);
 
 	/**
 	 * Places multiple assets based on the given FAssetPlacementInfos and FPlacementOptions.
 	 * @returns an array of FTypedElementHandles corresponding to any successfully placed elements.
 	 */
-	TArray<FTypedElementHandle> PlaceAssets(TArrayView<const FAssetPlacementInfo> InPlacementInfos, const FPlacementOptions& InPlacementOptions);
+	UE_API TArray<FTypedElementHandle> PlaceAssets(TArrayView<const FAssetPlacementInfo> InPlacementInfos, const FPlacementOptions& InPlacementOptions);
 
 	/**
 	 * Finds a registered AssetFactory for the given FAssetData.
 	 * @returns the first found factory, or nullptr if no valid factory is found.
 	 */
-	TScriptInterface<IAssetFactoryInterface> FindAssetFactoryFromAssetData(const FAssetData& InAssetData);
+	UE_API TScriptInterface<IAssetFactoryInterface> FindAssetFactoryFromAssetData(const FAssetData& InAssetData);
 
 	/**
 	 * Determines if the placement system is currently placing preview elements.
 	 * @returns true if the current PlaceAssets call is creating preview elements.
 	 */
-	bool IsCreatingPreviewElements() const;
+	UE_API bool IsCreatingPreviewElements() const;
 
-	FSimpleMulticastDelegate& OnPlacementFactoriesRegistered();
+	UE_API FSimpleMulticastDelegate& OnPlacementFactoriesRegistered();
 
-	TScriptInterface<IAssetFactoryInterface> GetAssetFactoryFromFactoryClass(UClass* InFactoryInterfaceClass) const;
+	UE_API TScriptInterface<IAssetFactoryInterface> GetAssetFactoryFromFactoryClass(UClass* InFactoryInterfaceClass) const;
 
 	UE_DEPRECATED(5.4, "This overload never worked as intended. Use the UClass* overload instead.")
-	TScriptInterface<IAssetFactoryInterface> GetAssetFactoryFromFactoryClass(TSubclassOf<UClass> InFactoryInterfaceClass) const;
+	UE_API TScriptInterface<IAssetFactoryInterface> GetAssetFactoryFromFactoryClass(TSubclassOf<UClass> InFactoryInterfaceClass) const;
 	
-	void RegisterAssetFactory(TScriptInterface<IAssetFactoryInterface> AssetFactory);
+	UE_API void RegisterAssetFactory(TScriptInterface<IAssetFactoryInterface> AssetFactory);
 
-	void UnregisterAssetFactory(TScriptInterface<IAssetFactoryInterface> AssetFactory);
+	UE_API void UnregisterAssetFactory(TScriptInterface<IAssetFactoryInterface> AssetFactory);
 	
 private:
-	void RegisterPlacementFactories();
-	void UnregisterPlacementFactories();
+	UE_API void RegisterPlacementFactories();
+	UE_API void UnregisterPlacementFactories();
 
 	UPROPERTY()
 	TArray<TScriptInterface<IAssetFactoryInterface>> AssetFactories;
@@ -145,3 +147,5 @@ private:
 
 	bool bIsCreatingPreviewElements = false;
 };
+
+#undef UE_API
diff --git a/Engine/Source/Editor/EditorFramework/Public/Toolkits/AssetEditorModeUILayer.h b/Engine/Source/Editor/EditorFramework/Public/Toolkits/AssetEditorModeUILayer.h
--- a/Engine/Source/Editor/EditorFramework/Public/Toolkits/AssetEditorModeUILayer.h
+++ b/Engine/Source/Editor/EditorFramework/Public/Toolkits/AssetEditorModeUILayer.h
@@ -28,6 +28,8 @@
 
 #include "AssetEditorModeUILayer.generated.h"
 
+#define UE_API EDITORFRAMEWORK_API
+
 class FExtender;
 class FLayoutExtender;
 class FWorkspaceItem;
@@ -38,40 +40,40 @@ class SDockTab;
 class UObject;
 struct FSlateIcon;
 
-UCLASS()
-class EDITORFRAMEWORK_API UAssetEditorUISubsystem : public UEditorSubsystem
+UCLASS(MinimalAPI)
+class UAssetEditorUISubsystem : public UEditorSubsystem
 {
 	GENERATED_BODY()
 
 public:
-	static const FName VerticalToolbarID;
-	static const FName TopLeftTabID;
-	static const FName BottomLeftTabID;
-	static const FName TopRightTabID;
-	static const FName BottomRightTabID;
+	static UE_API const FName VerticalToolbarID;
+	static UE_API const FName TopLeftTabID;
+	static UE_API const FName BottomLeftTabID;
+	static UE_API const FName TopRightTabID;
+	static UE_API const FName BottomRightTabID;
 
 protected:
 	virtual void RegisterLayoutExtensions(FLayoutExtender& Extender) {};
 };
 
-class EDITORFRAMEWORK_API FAssetEditorModeUILayer : public TSharedFromThis<FAssetEditorModeUILayer>
+class FAssetEditorModeUILayer : public TSharedFromThis<FAssetEditorModeUILayer>
 {
 public:
-	FAssetEditorModeUILayer(const IToolkitHost* InToolkitHost);
+	UE_API FAssetEditorModeUILayer(const IToolkitHost* InToolkitHost);
 	FAssetEditorModeUILayer() {};
 	virtual ~FAssetEditorModeUILayer() {};
 	/** Called by SLevelEditor to notify the toolbox about a new toolkit being hosted */
-	virtual void OnToolkitHostingStarted(const TSharedRef<IToolkit>& Toolkit);
+	UE_API virtual void OnToolkitHostingStarted(const TSharedRef<IToolkit>& Toolkit);
 
 	/** Called by SLevelEditor to notify the toolbox about an existing toolkit no longer being hosted */
-	virtual void OnToolkitHostingFinished(const TSharedRef<IToolkit>& Toolkit);
-	virtual TSharedPtr<FTabManager> GetTabManager();
-	virtual TSharedPtr<FWorkspaceItem> GetModeMenuCategory() const;
-	virtual void SetModePanelInfo(const FName InTabSpawnerID, const FMinorTabConfig& InTabInfo);
-	virtual TMap<FName, TWeakPtr<SDockTab>> GetSpawnedTabs();
+	UE_API virtual void OnToolkitHostingFinished(const TSharedRef<IToolkit>& Toolkit);
+	UE_API virtual TSharedPtr<FTabManager> GetTabManager();
+	UE_API virtual TSharedPtr<FWorkspaceItem> GetModeMenuCategory() const;
+	UE_API virtual void SetModePanelInfo(const FName InTabSpawnerID, const FMinorTabConfig& InTabInfo);
+	UE_API virtual TMap<FName, TWeakPtr<SDockTab>> GetSpawnedTabs();
 
 	/** Called by the Toolkit Host to set the name of the ModeToolbar this mode will extend */
-	virtual void SetSecondaryModeToolbarName(FName InName);
+	UE_API virtual void SetSecondaryModeToolbarName(FName InName);
 	
 	virtual FSimpleDelegate& ToolkitHostReadyForUI()
 	{
@@ -100,14 +102,14 @@ public:
 	}
 
 protected:
-	const FOnSpawnTab& GetStoredSpawner(const FName TabID);
-	void RegisterModeTabSpawners();
-	void RegisterModeTabSpawner(const FName TabID);
-	TSharedRef<SDockTab> SpawnStoredTab(const FSpawnTabArgs& Args, const FName TabID);
-	bool CanSpawnStoredTab(const FSpawnTabArgs& Args, const FName TabID);
-	FText GetTabSpawnerName(const FName TabID) const;
-	FText GetTabSpawnerTooltip(const FName TabID) const;
-	const FSlateIcon& GetTabSpawnerIcon(const FName TabID) const;
+	UE_API const FOnSpawnTab& GetStoredSpawner(const FName TabID);
+	UE_API void RegisterModeTabSpawners();
+	UE_API void RegisterModeTabSpawner(const FName TabID);
+	UE_API TSharedRef<SDockTab> SpawnStoredTab(const FSpawnTabArgs& Args, const FName TabID);
+	UE_API bool CanSpawnStoredTab(const FSpawnTabArgs& Args, const FName TabID);
+	UE_API FText GetTabSpawnerName(const FName TabID) const;
+	UE_API FText GetTabSpawnerTooltip(const FName TabID) const;
+	UE_API const FSlateIcon& GetTabSpawnerIcon(const FName TabID) const;
 
 protected:
 	/** The host of the toolkits created by modes */
@@ -128,3 +130,5 @@ protected:
 	/* A list of commands this ModeUILayer is aware of, currently only passed into the Mode Toolbar */
 	TSharedRef<FUICommandList> ModeCommands;
 };
+
+#undef UE_API
diff --git a/Engine/Source/Editor/EditorFramework/Public/Tools/AssetEditorContextInterface.h b/Engine/Source/Editor/EditorFramework/Public/Tools/AssetEditorContextInterface.h
--- a/Engine/Source/Editor/EditorFramework/Public/Tools/AssetEditorContextInterface.h
+++ b/Engine/Source/Editor/EditorFramework/Public/Tools/AssetEditorContextInterface.h
@@ -19,7 +19,7 @@ class UAssetEditorContextInterface : public UInterface
 /**
  * An interface for a context object that exposes common functionality of an asset editor.
  */
-class EDITORFRAMEWORK_API IAssetEditorContextInterface
+class IAssetEditorContextInterface
 {	
 	GENERATED_BODY()
 
diff --git a/Engine/Source/Editor/EditorFramework/Public/ViewportTabContent.h b/Engine/Source/Editor/EditorFramework/Public/ViewportTabContent.h
--- a/Engine/Source/Editor/EditorFramework/Public/ViewportTabContent.h
+++ b/Engine/Source/Editor/EditorFramework/Public/ViewportTabContent.h
@@ -12,6 +12,8 @@
 #include "Templates/SharedPointer.h"
 #include "UObject/NameTypes.h"
 
+#define UE_API EDITORFRAMEWORK_API
+
 /**
  * Represents the content in a viewport tab in an editor.
  * Each SDockTab holding viewports in an editor contains and owns one of these.
@@ -22,16 +24,16 @@ class FUICommandList;
 class IEditorViewportLayoutEntity;
 class SDockTab;
 
-class EDITORFRAMEWORK_API FViewportTabContent
+class FViewportTabContent
 {
 public:
 	virtual ~FViewportTabContent() {}
 
 	/** Returns whether the tab is currently shown */
-	bool IsVisible() const;
+	UE_API bool IsVisible() const;
 
 	/** @return True if this viewport belongs to the tab given */
-	bool BelongsToTab(TSharedRef<SDockTab> InParentTab) const;
+	UE_API bool BelongsToTab(TSharedRef<SDockTab> InParentTab) const;
 
 	/**
 	* Returns whether the named layout is currently selected
@@ -39,7 +41,7 @@ public:
 	* @param ConfigurationName		The name of the layout (for the names in namespace LevelViewportConfigurationNames)
 	* @return						True, if the named layout is currently active
 	*/
-	bool IsViewportConfigurationSet(const FName& ConfigurationName) const;
+	UE_API bool IsViewportConfigurationSet(const FName& ConfigurationName) const;
 
 	virtual void SetViewportConfiguration(const FName& ConfigurationName) {}
 
@@ -49,10 +51,10 @@ public:
 	 */
 	virtual void BindViewportLayoutCommands(FUICommandList& InOutCommandList, FName ViewportConfigKey) {}
 
-	const TMap< FName, TSharedPtr< IEditorViewportLayoutEntity > >* GetViewports() const;
+	UE_API const TMap< FName, TSharedPtr< IEditorViewportLayoutEntity > >* GetViewports() const;
 
 	using ViewportActionFunction = TFunction<void(FName Name, TSharedPtr<IEditorViewportLayoutEntity>)>;
-	void PerformActionOnViewports(ViewportActionFunction& TFuncPtr);
+	UE_API void PerformActionOnViewports(ViewportActionFunction& TFuncPtr);
 
 	DECLARE_EVENT(FViewportTabContent, FViewportTabContentLayoutChangedEvent);
 	virtual FViewportTabContentLayoutChangedEvent& OnViewportTabContentLayoutChanged() { return OnViewportTabContentLayoutChangedEvent; };
@@ -73,3 +75,5 @@ protected:
 
 	TOptional<FName> PreviouslyFocusedViewport;
 };
+
+#undef UE_API
diff --git a/Engine/Source/Editor/EditorFramework/Public/Viewports/InViewportUIDragOperation.h b/Engine/Source/Editor/EditorFramework/Public/Viewports/InViewportUIDragOperation.h
--- a/Engine/Source/Editor/EditorFramework/Public/Viewports/InViewportUIDragOperation.h
+++ b/Engine/Source/Editor/EditorFramework/Public/Viewports/InViewportUIDragOperation.h
@@ -11,6 +11,8 @@
 #include "Templates/SharedPointer.h"
 #include "Widgets/DeclarativeSyntaxSupport.h"
 
+#define UE_API EDITORFRAMEWORK_API
+
 struct FPointerEvent;
 
 DECLARE_DELEGATE_OneParam(FOnInViewportUIDropped, const FVector2D DropLocation)
@@ -18,7 +20,7 @@ DECLARE_DELEGATE_OneParam(FOnInViewportUIDropped, const FVector2D DropLocation)
 class SDockingTabStack;
 
 /** A Sample implementation of IDragDropOperation */
-class EDITORFRAMEWORK_API FInViewportUIDragOperation : public FDragDropOperation
+class FInViewportUIDragOperation : public FDragDropOperation
 {
 public:
 
@@ -29,14 +31,14 @@ public:
 	*
 	* @param bDropWasHandled   true when the drop was handled by some widget; false otherwise
 	*/
-	virtual void OnDrop(bool bDropWasHandled, const FPointerEvent& MouseEvent) override;
+	UE_API virtual void OnDrop(bool bDropWasHandled, const FPointerEvent& MouseEvent) override;
 
 	/**
 	 * Called when the mouse was moved during a drag and drop operation
 	 *
 	 * @param DragDropEvent    The event that describes this drag drop operation.
 	 */
-	virtual void OnDragged(const FDragDropEvent& DragDropEvent) override;
+	UE_API virtual void OnDragged(const FDragDropEvent& DragDropEvent) override;
 
 	/**
 	 * Create this Drag and Drop Content
@@ -47,18 +49,18 @@ public:
 	 *
 	 * @return a new FDockingDragOperation
 	 */
-	static TSharedRef<FInViewportUIDragOperation> New(const TSharedRef<class SWidget>& InUIToBeDragged, const FVector2D InDecoratorOffset, const FVector2D& OwnerAreaSize, FOnInViewportUIDropped& InDropDelegate);
+	static UE_API TSharedRef<FInViewportUIDragOperation> New(const TSharedRef<class SWidget>& InUIToBeDragged, const FVector2D InDecoratorOffset, const FVector2D& OwnerAreaSize, FOnInViewportUIDropped& InDropDelegate);
 
-	virtual ~FInViewportUIDragOperation();
+	UE_API virtual ~FInViewportUIDragOperation();
 
 	/** @return The offset into the tab where the user grabbed in Slate Units. */
-	const FVector2f GetDecoratorOffsetFromCursor();
+	UE_API const FVector2f GetDecoratorOffsetFromCursor();
 
-	void BroadcastDropEvent(const FVector2D InLocation);
+	UE_API void BroadcastDropEvent(const FVector2D InLocation);
 
 protected:
 	/** The constructor is protected, so that this class can only be instantiated as a shared pointer. */
-	FInViewportUIDragOperation(const TSharedRef<class SWidget>& InUIToBeDragged, const FVector2D InDecoratorOffsetFromCursor, const FVector2D& OwnerAreaSize, FOnInViewportUIDropped& InDropDelegate);
+	UE_API FInViewportUIDragOperation(const TSharedRef<class SWidget>& InUIToBeDragged, const FVector2D InDecoratorOffsetFromCursor, const FVector2D& OwnerAreaSize, FOnInViewportUIDropped& InDropDelegate);
 
 	/** What is actually being dragged in this operation */
 	TSharedPtr<class SWidget> UIBeingDragged;
@@ -74,3 +76,5 @@ protected:
 
 	FOnInViewportUIDropped DropDelegate;
 };
+
+#undef UE_API
